Copying //UE4/Dev-Physics to Dev-Main (//UE4/Dev-Main) [at] 9098278

diff --git a/Engine/Source/Editor/Persona/Private/PersonaMeshDetails.cpp b/Engine/Source/Editor/Persona/Private/PersonaMeshDetails.cpp
--- a/Engine/Source/Editor/Persona/Private/PersonaMeshDetails.cpp
+++ b/Engine/Source/Editor/Persona/Private/PersonaMeshDetails.cpp
@@ -3834,7 +3834,7 @@ void FPersonaMeshDetails::OnGenerateElementForClothingAsset( TSharedRef<IPropert
 	.HAlign(HAlign_Fill)
 	[
 		SNew(STextBlock)
-		.Text(FText::FromString(CurrentAsset->GetName()))
+		.Text(CurrentAsset ? FText::FromString(CurrentAsset->GetName()) : FText())
 	];
 
 	ChildrenBuilder.AddCustomRow(LOCTEXT("ClothingAsset_Search_Details", "Details"))
@@ -3879,22 +3879,27 @@ void FPersonaMeshDetails::OnGenerateElementForClothingAsset( TSharedRef<IPropert
 	Options.bHideFilterArea = true;
 	Options.bShowComponents = false;
 
-	Inspector->ShowDetailsForSingleObject(CurrentAsset, Options);
+	if (CurrentAsset)
+	{
+		Inspector->ShowDetailsForSingleObject(CurrentAsset, Options);
+	}
 }
 
 TSharedRef<SUniformGridPanel> FPersonaMeshDetails::MakeClothingDetailsWidget(int32 AssetIndex) const
 {
-	const FSlateFontInfo DetailFontInfo = IDetailLayoutBuilder::GetDetailFont();
+	TSharedRef<SUniformGridPanel> Grid = SNew(SUniformGridPanel).SlotPadding(2.0f);
 
 	USkeletalMesh* SkelMesh = GetPersonaToolkit()->GetMesh();
-	UClothingAsset* ClothingAsset = Cast<UClothingAsset>(SkelMesh->MeshClothingAssets[AssetIndex]);
-	check(ClothingAsset);
+	UClothingAssetBase* ClothingAsset = SkelMesh->MeshClothingAssets[AssetIndex];
+	if (!ClothingAsset)
+	{
+		return Grid;
+	}
 
-	TSharedRef<SUniformGridPanel> Grid = SNew(SUniformGridPanel).SlotPadding(2.0f);
 
-	int32 NumLODs = ClothingAsset->LodData.Num();
+	const FSlateFontInfo DetailFontInfo = IDetailLayoutBuilder::GetDetailFont();
+	const int32 NumLODs = ClothingAsset->GetNumLods();
 	int32 RowNumber = 0;
-
 	for(int32 LODIndex=0; LODIndex < NumLODs; LODIndex++)
 	{
 		Grid->AddSlot(0, RowNumber) // x, y
@@ -3907,91 +3912,99 @@ TSharedRef<SUniformGridPanel> FPersonaMeshDetails::MakeClothingDetailsWidget(int
 
 		RowNumber++;
 
-		FClothLODData& LodData = ClothingAsset->LodData[LODIndex];
-		FClothPhysicalMeshData& PhysMeshData = LodData.PhysicalMeshData;
-		FClothCollisionData& CollisionData = LodData.CollisionData;
+		if (UClothingAsset* Asset = Cast<UClothingAsset>(ClothingAsset))
+		{
+			FClothLODData& LodData = Asset->LodData[LODIndex];
+			FClothPhysicalMeshData& PhysMeshData = LodData.PhysicalMeshData;
+			FClothCollisionData& CollisionData = LodData.CollisionData;
 
-		Grid->AddSlot(0, RowNumber) 
-			.HAlign(HAlign_Center)
-			[
-				SNew(STextBlock)
-				.Font(DetailFontInfo)
+			Grid->AddSlot(0, RowNumber)
+				.HAlign(HAlign_Center)
+				[
+					SNew(STextBlock)
+					.Font(DetailFontInfo)
 				.Text(LOCTEXT("SimulVertexCount", "Simul Verts"))
-			];
+				];
 
-		Grid->AddSlot(0, RowNumber + 1)
-			.HAlign(HAlign_Center)
-			[
-				SNew(STextBlock)
-				.Font(DetailFontInfo)
+			Grid->AddSlot(0, RowNumber + 1)
+				.HAlign(HAlign_Center)
+				[
+					SNew(STextBlock)
+					.Font(DetailFontInfo)
 				.Text(FText::AsNumber(PhysMeshData.Vertices.Num() - PhysMeshData.NumFixedVerts))
-			];
+				];
 
-		Grid->AddSlot(1, RowNumber)
-			.HAlign(HAlign_Center)
-			[
-				SNew(STextBlock)
-				.Font(DetailFontInfo)
+			Grid->AddSlot(1, RowNumber)
+				.HAlign(HAlign_Center)
+				[
+					SNew(STextBlock)
+					.Font(DetailFontInfo)
 				.Text(LOCTEXT("FixedVertexCount", "Fixed Verts"))
-			];
+				];
 
-		Grid->AddSlot(1, RowNumber + 1)
-			.HAlign(HAlign_Center)
-			[
-				SNew(STextBlock)
-				.Font(DetailFontInfo)
+			Grid->AddSlot(1, RowNumber + 1)
+				.HAlign(HAlign_Center)
+				[
+					SNew(STextBlock)
+					.Font(DetailFontInfo)
 				.Text(FText::AsNumber(PhysMeshData.NumFixedVerts))
-			];
+				];
 
-		Grid->AddSlot(2, RowNumber)
-			.HAlign(HAlign_Center)
-			[
-				SNew(STextBlock)
-				.Font(DetailFontInfo)
+			Grid->AddSlot(2, RowNumber)
+				.HAlign(HAlign_Center)
+				[
+					SNew(STextBlock)
+					.Font(DetailFontInfo)
 				.Text(LOCTEXT("TriangleCount", "Sim Triangles"))
-			];
+				];
 
-		Grid->AddSlot(2, RowNumber + 1)
-			.HAlign(HAlign_Center)
-			[
-				SNew(STextBlock)
-				.Font(DetailFontInfo)
+			Grid->AddSlot(2, RowNumber + 1)
+				.HAlign(HAlign_Center)
+				[
+					SNew(STextBlock)
+					.Font(DetailFontInfo)
 				.Text(FText::AsNumber(PhysMeshData.Indices.Num() / 3))
-			];
+				];
 
-		Grid->AddSlot(3, RowNumber)
-			.HAlign(HAlign_Center)
-			[
-				SNew(STextBlock)
-				.Font(DetailFontInfo)
+			Grid->AddSlot(3, RowNumber)
+				.HAlign(HAlign_Center)
+				[
+					SNew(STextBlock)
+					.Font(DetailFontInfo)
 				.Text(LOCTEXT("NumUsedBones", "Bones"))
-			];
+				];
 
-		Grid->AddSlot(3, RowNumber + 1)
-			.HAlign(HAlign_Center)
-			[
-				SNew(STextBlock)
-				.Font(DetailFontInfo)
+			Grid->AddSlot(3, RowNumber + 1)
+				.HAlign(HAlign_Center)
+				[
+					SNew(STextBlock)
+					.Font(DetailFontInfo)
 				.Text(FText::AsNumber(PhysMeshData.MaxBoneWeights))
-			];
+				];
 
-		Grid->AddSlot(4, RowNumber)
-			.HAlign(HAlign_Center)
-			[
-				SNew(STextBlock)
-				.Font(DetailFontInfo)
+			Grid->AddSlot(4, RowNumber)
+				.HAlign(HAlign_Center)
+				[
+					SNew(STextBlock)
+					.Font(DetailFontInfo)
 				.Text(LOCTEXT("NumBoneSpheres", "Spheres"))
-			];
+				];
 
-		Grid->AddSlot(4, RowNumber + 1)
-			.HAlign(HAlign_Center)
-			[
-				SNew(STextBlock)
-				.Font(DetailFontInfo)
+			Grid->AddSlot(4, RowNumber + 1)
+				.HAlign(HAlign_Center)
+				[
+					SNew(STextBlock)
+					.Font(DetailFontInfo)
 				.Text(FText::AsNumber(CollisionData.Spheres.Num()))
-			];
+				];
 
-		RowNumber += 2;
+			RowNumber += 2;
+		}
+		else
+		{
+			// Unsupported asset type
+			check(false);
+		}
 	}
 
 	return Grid;
@@ -4079,11 +4092,10 @@ FReply FPersonaMeshDetails::OnRemoveApexFileClicked(int32 AssetIndex, IDetailLay
 		// Now we can remove the asset.
 		if(SkelMesh->MeshClothingAssets.IsValidIndex(AssetIndex))
 		{
-			UClothingAssetBase* AssetToRemove = SkelMesh->MeshClothingAssets[AssetIndex];
-			check(AssetToRemove);
-
-			AssetToRemove->UnbindFromSkeletalMesh(SkelMesh);
-
+			if (UClothingAssetBase* AssetToRemove = SkelMesh->MeshClothingAssets[AssetIndex])
+			{
+				AssetToRemove->UnbindFromSkeletalMesh(SkelMesh);
+			}
 			SkelMesh->MeshClothingAssets.RemoveAt(AssetIndex);
 
 			// Need to fix up asset indices on sections.
@@ -4147,18 +4159,17 @@ void FPersonaMeshDetails::UpdateClothingEntries()
 	const int32 NumClothingAssets = Mesh->MeshClothingAssets.Num();
 	for(int32 Idx = 0; Idx < NumClothingAssets; ++Idx)
 	{
-		UClothingAsset* Asset = CastChecked<UClothingAsset>(Mesh->MeshClothingAssets[Idx]);
-
-		const int32 NumAssetLods = Asset->LodData.Num();
-		for(int32 AssetLodIndex = 0; AssetLodIndex < NumAssetLods; ++AssetLodIndex)
+		if (UClothingAssetBase* ClothingAsset = Mesh->MeshClothingAssets[Idx])
 		{
-			TSharedPtr<FClothingEntry> NewEntry = MakeShared<FClothingEntry>();
-
-			NewEntry->Asset = Mesh->MeshClothingAssets[Idx];
-			NewEntry->AssetIndex = Idx;
-			NewEntry->AssetLodIndex = AssetLodIndex;
-
-			NewClothingAssetEntries.Add(NewEntry);
+			const int32 NumAssetLods = ClothingAsset->GetNumLods();
+			for (int32 AssetLodIndex = 0; AssetLodIndex < NumAssetLods; ++AssetLodIndex)
+			{
+				TSharedPtr<FClothingEntry> NewEntry = MakeShared<FClothingEntry>();
+				NewEntry->Asset = ClothingAsset;
+				NewEntry->AssetIndex = Idx;
+				NewEntry->AssetLodIndex = AssetLodIndex;
+				NewClothingAssetEntries.Add(NewEntry);
+			}
 		}
 	}
 }
@@ -4265,9 +4276,8 @@ void FPersonaMeshDetails::OnFinishedChangingClothingProperties(const FPropertyCh
 			USkeletalMesh* CurrentMesh = GetPersonaToolkit()->GetMesh();
 			if(CurrentMesh->MeshClothingAssets.IsValidIndex(InAssetIndex))
 			{
-				UClothingAsset* Asset = CastChecked<UClothingAsset>(CurrentMesh->MeshClothingAssets[InAssetIndex]);
-
-				Asset->BuildSelfCollisionData();
+				if(UClothingAssetBase* Asset = CurrentMesh->MeshClothingAssets[InAssetIndex])
+					Asset->BuildSelfCollisionData();
 			}
 		}
 	}
diff --git a/Engine/Source/Editor/Persona/Private/SAnimViewportToolBar.cpp b/Engine/Source/Editor/Persona/Private/SAnimViewportToolBar.cpp
--- a/Engine/Source/Editor/Persona/Private/SAnimViewportToolBar.cpp
+++ b/Engine/Source/Editor/Persona/Private/SAnimViewportToolBar.cpp
@@ -34,7 +34,7 @@
 #include "PreviewSceneCustomizations.h"
 #include "ClothingSimulation.h"
 #include "SimulationEditorExtender.h"
-#include "ClothingSimulationFactoryInterface.h"
+#include "ClothingSimulationFactory.h"
 #include "ClothingSystemEditorInterfaceModule.h"
 #include "Widgets/SWidget.h"
 #include "Types/ISlateMetaData.h"
diff --git a/Engine/Source/Editor/Persona/Private/SCreateClothingSettingsPanel.cpp b/Engine/Source/Editor/Persona/Private/SCreateClothingSettingsPanel.cpp
--- a/Engine/Source/Editor/Persona/Private/SCreateClothingSettingsPanel.cpp
+++ b/Engine/Source/Editor/Persona/Private/SCreateClothingSettingsPanel.cpp
@@ -13,7 +13,7 @@
 #include "DetailCategoryBuilder.h"
 #include "DetailWidgetRow.h"
 #include "PropertyCustomizationHelpers.h"
-#include "ClothingAssetInterface.h"
+#include "ClothingAssetBase.h"
 #include "Widgets/Input/SComboButton.h"
 #include "Widgets/Text/STextBlock.h"
 #include "Framework/MultiBox/MultiBoxBuilder.h"
@@ -152,11 +152,12 @@ TSharedRef<SWidget> FClothCreateSettingsCustomization::OnGetTargetAssetMenu()
 			const int32 NumCloths = Mesh->MeshClothingAssets.Num();
 			for(int32 ClothIndex = 0 ; ClothIndex < NumCloths ; ++ClothIndex)
 			{
-				UClothingAssetBase* ClothingAsset = Mesh->MeshClothingAssets[ClothIndex];
-
-				FUIAction Action;
-				Action.ExecuteAction = FExecuteAction::CreateSP(this, &FClothCreateSettingsCustomization::OnAssetSelected, ClothIndex);
-				Builder.AddMenuEntry(FText::FromString(ClothingAsset->GetName()), LOCTEXT("TargetAsset_Tooltip", "Select this clothing as the target to import to."), FSlateIcon(), Action);
+				if (UClothingAssetBase* ClothingAsset = Mesh->MeshClothingAssets[ClothIndex])
+				{
+					FUIAction Action;
+					Action.ExecuteAction = FExecuteAction::CreateSP(this, &FClothCreateSettingsCustomization::OnAssetSelected, ClothIndex);
+					Builder.AddMenuEntry(FText::FromString(ClothingAsset->GetName()), LOCTEXT("TargetAsset_Tooltip", "Select this clothing as the target to import to."), FSlateIcon(), Action);
+				}
 			}
 		}
 	}
@@ -185,7 +186,8 @@ void FClothCreateSettingsCustomization::OnAssetSelected(int32 InMeshClothingInde
 
 	if(USkeletalMesh* Mesh = MeshPtr.Get())
 	{
-		if(Mesh->MeshClothingAssets.IsValidIndex(InMeshClothingIndex))
+		if(Mesh->MeshClothingAssets.IsValidIndex(InMeshClothingIndex) && 
+		   Mesh->MeshClothingAssets[InMeshClothingIndex])
 		{
 			if(ParamsStruct->TargetAsset != Mesh->MeshClothingAssets[InMeshClothingIndex])
 			{
diff --git a/Engine/Source/Editor/UnrealEd/Classes/Factories/Factory.h b/Engine/Source/Editor/UnrealEd/Classes/Factories/Factory.h
--- a/Engine/Source/Editor/UnrealEd/Classes/Factories/Factory.h
+++ b/Engine/Source/Editor/UnrealEd/Classes/Factories/Factory.h
@@ -160,7 +160,7 @@ public:
 	virtual bool ImportUntypedBulkDataFromText(const TCHAR*& Buffer, FUntypedBulkData& BulkData);
 
 	/** Creates a list of file extensions supported by this factory */
-	void GetSupportedFileExtensions(TArray<FString>& OutExtensions) const;
+	virtual void GetSupportedFileExtensions(TArray<FString>& OutExtensions) const;
 
 	/** Do clean up after importing is done. Will be called once for multi batch import. */
 	virtual void CleanUp() {}
diff --git a/Engine/Source/Editor/UnrealEd/Private/Animation/DebugSkelMeshComponent.cpp b/Engine/Source/Editor/UnrealEd/Private/Animation/DebugSkelMeshComponent.cpp
--- a/Engine/Source/Editor/UnrealEd/Private/Animation/DebugSkelMeshComponent.cpp
+++ b/Engine/Source/Editor/UnrealEd/Private/Animation/DebugSkelMeshComponent.cpp
@@ -19,6 +19,7 @@
 
 #include "Assets/ClothingAsset.h"
 #include "ClothingSimulation.h"
+#include "Utils/ClothingMeshUtils.h"
 #include "DynamicMeshBuilder.h"
 #include "Materials/MaterialInstanceDynamic.h"
 
@@ -813,7 +814,7 @@ void UDebugSkelMeshComponent::RefreshSelectedClothingSkinnedPositions()
 	{
 		UClothingAssetBase** Asset = SkeletalMesh->MeshClothingAssets.FindByPredicate([&](UClothingAssetBase* Item)
 		{
-			return SelectedClothingGuidForPainting == Item->GetAssetGuid();
+			return Item && SelectedClothingGuidForPainting == Item->GetAssetGuid();
 		});
 
 		if(Asset)
@@ -831,7 +832,7 @@ void UDebugSkelMeshComponent::RefreshSelectedClothingSkinnedPositions()
 
 				FClothLODData& LodData = ConcreteAsset->LodData[SelectedClothingLodForPainting];
 
-				FClothingSimulationBase::SkinPhysicsMesh(ConcreteAsset, LodData.PhysicalMeshData, FTransform::Identity, RefToLocals.GetData(), RefToLocals.Num(), SkinnedSelectedClothingPositions, SkinnedSelectedClothingNormals);
+				ClothingMeshUtils::SkinPhysicsMesh(ConcreteAsset->UsedBoneIndices, LodData.PhysicalMeshData, FTransform::Identity, RefToLocals.GetData(), RefToLocals.Num(), SkinnedSelectedClothingPositions, SkinnedSelectedClothingNormals);
 				RebuildCachedClothBounds();
 			}
 		}
@@ -1028,7 +1029,7 @@ FDebugSkelMeshDynamicData::FDebugSkelMeshDynamicData(UDebugSkelMeshComponent* In
 			for(int32 ClothingAssetIndex = 0; ClothingAssetIndex < NumClothingAssets; ++ClothingAssetIndex)
 			{
 				UClothingAssetBase* BaseAsset = Mesh->MeshClothingAssets[ClothingAssetIndex];
-				if(BaseAsset->GetAssetGuid() == InComponent->SelectedClothingGuidForPainting)
+				if(BaseAsset && BaseAsset->GetAssetGuid() == InComponent->SelectedClothingGuidForPainting)
 				{
 					ClothingSimDataIndexWhenPainting = ClothingAssetIndex;
 
diff --git a/Engine/Source/Editor/UnrealEd/Private/Factories/SkeletalMeshImport.cpp b/Engine/Source/Editor/UnrealEd/Private/Factories/SkeletalMeshImport.cpp
--- a/Engine/Source/Editor/UnrealEd/Private/Factories/SkeletalMeshImport.cpp
+++ b/Engine/Source/Editor/UnrealEd/Private/Factories/SkeletalMeshImport.cpp
@@ -26,7 +26,7 @@
 #include "LODUtilities.h"
 #include "UObject/Package.h"
 #include "MeshUtilities.h"
-#include "ClothingAssetInterface.h"
+#include "ClothingAssetBase.h"
 #include "Factories/FbxSkeletalMeshImportData.h"
 #include "IMeshReductionManagerModule.h"
 #include "Rendering/SkeletalMeshModel.h"
@@ -930,7 +930,8 @@ void RestoreExistingSkelMeshData(ExistingSkelMeshData* MeshData, USkeletalMesh*
 
 		for (UClothingAssetBase* ClothingAsset : SkeletalMesh->MeshClothingAssets)
 		{
-			ClothingAsset->RefreshBoneMapping(SkeletalMesh);
+			if(ClothingAsset)
+				ClothingAsset->RefreshBoneMapping(SkeletalMesh);
 		}
 
 		SkeletalMesh->SetSamplingInfo(MeshData->ExistingSamplingInfo);
diff --git a/Engine/Source/Runtime/ClothingSystemRuntimeInterface/Private/ClothingSimulation.cpp b/Engine/Source/Runtime/ClothingSystemRuntimeInterface/Private/ClothingSimulation.cpp
new file mode 100644
--- /dev/null
+++ b/Engine/Source/Runtime/ClothingSystemRuntimeInterface/Private/ClothingSimulation.cpp
@@ -0,0 +1,9 @@
+// Copyright 1998-2019 Epic Games, Inc. All Rights Reserved.
+
+#include "ClothingSimulation.h"
+
+IClothingSimulation::IClothingSimulation()
+{}
+
+IClothingSimulation::~IClothingSimulation()
+{}
diff --git a/Engine/Source/Runtime/ClothingSystemRuntimeInterface/Private/ClothingSimulationContext.cpp b/Engine/Source/Runtime/ClothingSystemRuntimeInterface/Private/ClothingSimulationContext.cpp
new file mode 100644
--- /dev/null
+++ b/Engine/Source/Runtime/ClothingSystemRuntimeInterface/Private/ClothingSimulationContext.cpp
@@ -0,0 +1,28 @@
+// Copyright 1998-2019 Epic Games, Inc. All Rights Reserved.
+
+#include "ClothingSimulationContext.h"
+
+//==============================================================================
+// IClotingSimulationContext
+//==============================================================================
+
+IClothingSimulationContext::IClothingSimulationContext()
+{}
+
+IClothingSimulationContext::~IClothingSimulationContext()
+{}
+
+//==============================================================================
+// IClotingSimulationContextBase
+//==============================================================================
+
+FClothingSimulationContextBase::FClothingSimulationContextBase()
+	: DeltaSeconds(0.0f)
+	, PredictedLod(0)
+	, WindVelocity(FVector::ZeroVector)
+	, WindAdaption(0.0f)
+	, TeleportMode(EClothingTeleportMode::None)
+{}
+
+FClothingSimulationContextBase::~FClothingSimulationContextBase()
+{}
diff --git a/Engine/Source/Runtime/ClothingSystemRuntimeInterface/Private/ClothingSystemRuntimeTypes.cpp b/Engine/Source/Runtime/ClothingSystemRuntimeInterface/Private/ClothingSystemRuntimeTypes.cpp
--- a/Engine/Source/Runtime/ClothingSystemRuntimeInterface/Private/ClothingSystemRuntimeTypes.cpp
+++ b/Engine/Source/Runtime/ClothingSystemRuntimeInterface/Private/ClothingSystemRuntimeTypes.cpp
@@ -1,34 +1,3 @@
 // Copyright 1998-2019 Epic Games, Inc. All Rights Reserved.
 
 #include "ClothingSystemRuntimeTypes.h"
-
-void FClothCollisionData::Reset()
-{
-	Spheres.Reset();
-	SphereConnections.Reset();
-	Convexes.Reset();
-}
-
-void FClothCollisionData::Append(const FClothCollisionData& InOther)
-{
-	const int32 NumSpheresBefore = Spheres.Num();
-	const int32 NumSphereConnectionsBefore = SphereConnections.Num();
-
-	Spheres.Append(InOther.Spheres);
-	SphereConnections.Append(InOther.SphereConnections);
-
-	const int32 NumSphereConnectionsAfter = SphereConnections.Num();
-
-	if(NumSpheresBefore > 0)
-	{
-		// Each connection that was added needs to have its sphere indices increased to match the new spheres that were added
-		for(int32 NewConnectionIndex = NumSphereConnectionsBefore; NewConnectionIndex < NumSphereConnectionsAfter; ++NewConnectionIndex)
-		{
-			FClothCollisionPrim_SphereConnection& Connection = SphereConnections[NewConnectionIndex];
-			Connection.SphereIndices[0] += NumSpheresBefore;
-			Connection.SphereIndices[1] += NumSpheresBefore;
-		}
-	}
-
-	Convexes.Append(InOther.Convexes);
-}
\ No newline at end of file
diff --git a/Engine/Source/Runtime/ClothingSystemRuntimeInterface/Public/ClothCollisionData.h b/Engine/Source/Runtime/ClothingSystemRuntimeInterface/Public/ClothCollisionData.h
new file mode 100644
--- /dev/null
+++ b/Engine/Source/Runtime/ClothingSystemRuntimeInterface/Public/ClothCollisionData.h
@@ -0,0 +1,30 @@
+// Copyright 1998-2019 Epic Games, Inc. All Rights Reserved.
+
+#pragma once
+
+#include "ClothCollisionPrim.h"
+#include "Containers/Array.h"
+#include "UObject/ObjectMacros.h"
+#include "ClothCollisionData.generated.h"
+
+USTRUCT()
+struct CLOTHINGSYSTEMRUNTIMEINTERFACE_API FClothCollisionData
+{
+	GENERATED_BODY()
+
+	void Reset();
+
+	void Append(const FClothCollisionData& InOther);
+
+	// Sphere data
+	UPROPERTY(EditAnywhere, Category = Collison)
+	TArray<FClothCollisionPrim_Sphere> Spheres;
+
+	// Capsule data
+	UPROPERTY(EditAnywhere, Category = Collison)
+	TArray<FClothCollisionPrim_SphereConnection> SphereConnections;
+
+	// Convex Data
+	UPROPERTY(EditAnywhere, Category = Collison)
+	TArray<FClothCollisionPrim_Convex> Convexes;
+};
diff --git a/Engine/Source/Runtime/ClothingSystemRuntimeInterface/Public/ClothingAssetBase.h b/Engine/Source/Runtime/ClothingSystemRuntimeInterface/Public/ClothingAssetBase.h
new file mode 100644
--- /dev/null
+++ b/Engine/Source/Runtime/ClothingSystemRuntimeInterface/Public/ClothingAssetBase.h
@@ -0,0 +1,90 @@
+// Copyright 1998-2019 Epic Games, Inc. All Rights Reserved.
+
+#pragma once
+
+#include "UObject/Object.h"
+#include "Misc/Guid.h"
+#include "ClothingAssetBase.generated.h"
+
+class USkeletalMesh;
+
+// An interface object for any clothing asset the engine can use. Any
+// clothing asset concrete object should derive from this.
+UCLASS(Abstract)
+class CLOTHINGSYSTEMRUNTIMEINTERFACE_API UClothingAssetBase : public UObject
+{
+	GENERATED_BODY()
+
+public:
+
+	/** Binds a clothing asset submesh to a skeletal mesh section
+	* @param InSkelMesh Skel mesh to bind to
+	* @param InSectionIndex Section in the skel mesh to replace
+	* @param InSubmeshIdx Submesh in this asset to replace section with
+	* @param InAssetLodIndex Internal clothing LOD to use
+	*/
+	virtual bool BindToSkeletalMesh(USkeletalMesh* InSkelMesh, int32 InMeshLodIndex, int32 InSectionIndex, int32 InAssetLodIndex, bool bCallPostEditChange = true)
+	PURE_VIRTUAL(UClothingAssetBase::BindToSkeletalMesh, return false;);
+
+	/**
+	* Unbinds this clothing asset from the provided skeletal mesh, will remove all LODs
+	* @param InSkelMesh skeletal mesh to remove from
+	*/
+	virtual void UnbindFromSkeletalMesh(USkeletalMesh* InSkelMesh)
+	PURE_VIRTUAL(UClothingAssetBase::UnbindFromSkeletalMesh, );
+
+	/**
+	* Unbinds this clothing asset from the provided skeletal mesh
+	* @param InSkelMesh skeletal mesh to remove from
+	* @param InMeshLodIndex Mesh LOD to remove this asset from (could still be bound to other LODs)
+	*/
+	virtual void UnbindFromSkeletalMesh(USkeletalMesh* InSkelMesh, int32 InMeshLodIndex)
+	PURE_VIRTUAL(UClothingAssetBase::UnbindFromSkeletalMesh,);
+
+	/** 
+	 * Messages to the clothing asset that the bones in the parent mesh have
+	 * possibly changed, which could invalidate the bone indices stored in the LOD
+	 * data.
+	 * @param InSkelMesh - The mesh to use to remap the bones
+	 */
+	virtual void RefreshBoneMapping(USkeletalMesh* InSkelMesh)
+	PURE_VIRTUAL(UClothingAssetBase::RefreshBoneMapping, );
+
+	/**
+	 * Called on the clothing asset when the base data (physical mesh etc.) has
+	 * changed, so any intermediate generated data can be regenerated.
+	 */
+	virtual void InvalidateCachedData()
+	PURE_VIRTUAL(UClothingAssetBase::InvalidateCachedData(), );
+
+	/** Check the validity of a LOD index */
+	virtual bool IsValidLod(int32 InLodIndex)
+	PURE_VIRTUAL(UClothingAssetBase::IsValidLod(), return false;);
+
+	/** Get the number of LODs defined in the clothing asset */
+	virtual int32 GetNumLods()
+	PURE_VIRTUAL(UClothingAssetBase::GetNumLods(), return 0;);
+
+	/** Builds self collision data */
+	virtual void BuildSelfCollisionData()
+	PURE_VIRTUAL(UClothingAssetBase::BuildSelfCollisionData(), );
+
+	/** Get the guid identifying this asset */
+	const FGuid& GetAssetGuid() const
+	{
+		return AssetGuid;
+	}
+
+	// If this asset was imported from a file, this will be the original path
+	UPROPERTY(VisibleAnywhere, Category = Import)
+	FString ImportedFilePath;
+
+protected:
+
+	/** The asset factory should have access, as it will assign the asset guid when building assets */
+	friend class UClothingAssetFactory;
+
+	/** Guid to identify this asset. Will be embedded into chunks that are created using this asset */
+	UPROPERTY()
+	FGuid AssetGuid;
+};
diff --git a/Engine/Source/Runtime/ClothingSystemRuntimeInterface/Public/ClothingSimulation.h b/Engine/Source/Runtime/ClothingSystemRuntimeInterface/Public/ClothingSimulation.h
new file mode 100644
--- /dev/null
+++ b/Engine/Source/Runtime/ClothingSystemRuntimeInterface/Public/ClothingSimulation.h
@@ -0,0 +1,133 @@
+// Copyright 1998-2019 Epic Games, Inc. All Rights Reserved.
+#pragma once
+
+#include "ClothCollisionData.h"
+#include "ClothingSystemRuntimeTypes.h"
+
+class UClothingAssetBase;
+class IClothingSimulationContext;
+class USkeletalMeshComponent;
+
+/** Base class for clothing simulators. */
+class CLOTHINGSYSTEMRUNTIMEINTERFACE_API IClothingSimulation
+{
+public:
+	IClothingSimulation();
+	virtual ~IClothingSimulation();
+
+	// The majority of the API for this class is protected. The required objects (skel meshes and the parallel task)
+	// are friends so they can use the functionality. For the most part the simulation is not designed to be used
+	// outside of the skeletal mesh component as it parallel simulation is tied to the skel mesh tick and
+	// dependents.
+	// Any method that is available in the public section below should consider that it may be called while
+	// the simulation is running.
+	// Currently only a const pointer is exposed externally to the skeletal mesh component, so barring that
+	// being cast away external callers can't call non-const methods. The debug skel mesh component exposes a
+	// mutable version for editor use
+protected:
+
+	friend class USkeletalMeshComponent;
+	friend class FParallelClothTask;
+
+	/**
+	 * Create an actor for this simulation from the data in InAsset
+	 * Simulation data for this actor should be written back to SimDataIndex in GetSimulationData
+	 * @param InOwnerComponent - The component requesting this actor
+	 * @param InAsset - The asset to create an actor from
+	 * @param SimDataIndex - the sim data index to use when doing the writeback in GetSimulationData
+	 */
+	virtual void CreateActor(USkeletalMeshComponent* InOwnerComponent, UClothingAssetBase* InAsset, int32 SimDataIndex) = 0;
+	
+	/** Create a new context, will not be filled, call FillContext before simulating with this context */
+	virtual IClothingSimulationContext* CreateContext() = 0;
+
+	/**
+	 * Fills an existing context for a single simulation step, called by the engine on the game thread prior to simulation 
+	 * @param InComponent - The component to fill the context for
+	 * @param InOutContext - The context to fill
+	 */
+	virtual void FillContext(USkeletalMeshComponent* InComponent, float InDeltaTime, IClothingSimulationContext* InOutContext) = 0;
+
+	/** Initialize the simulation, will be called before any Simulate calls */
+	virtual void Initialize() = 0;
+
+	/**
+	 * Shutdown the simulation, this should clear ALL resources as we no longer expect to
+	 * call any other function on this simulation without first calling Initialize again.
+	 */
+	virtual void Shutdown() = 0;
+
+	/**
+	 * Called by the engine to detect whether or not the simulation should run, essentially
+	 * are there any actors that need to simulate in this simulation
+	 */
+	virtual bool ShouldSimulate() const = 0;
+
+	/**
+	 * Run a single tick of the simulation.
+	 * The pointer InContext is guaranteed (when called by the engine) to be the context allocated in
+	 * CreateContext and can be assumed to be safely castable to any derived type allocated there.
+	 * New callers should make care to make sure only the correct context is ever passed through
+	 * @param InContext - The context to use during simulation, will have been filled in FillContext
+	 */
+	virtual void Simulate(IClothingSimulationContext* InContext) = 0;
+
+	/** Simulation should remove all of it's actors when next possible and free them */
+	virtual void DestroyActors() = 0;
+
+	/** 
+	 * Destroy a context object, engine will always pass a context created using CreateContext 
+	 * @param InContext - The context to destroy
+	 */
+	virtual void DestroyContext(IClothingSimulationContext* InContext) = 0;
+
+	/** 
+	 * Fill FClothSimulData map for the clothing simulation. Should fill a map pair per-actor 
+	 * @param OutData - The simulation data to write to
+	 * @param InOwnerComponent - the component that owns the simulation
+	 * @param InOverrideComponent - An override component if bound to a master pose component
+	 */
+	virtual void GetSimulationData(TMap<int32, FClothSimulData>& OutData, USkeletalMeshComponent* InOwnerComponent, USkinnedMeshComponent* InOverrideComponent) const = 0;
+
+	/** Get the bounds of the simulation mesh in local simulation space */
+	virtual FBoxSphereBounds GetBounds(const USkeletalMeshComponent* InOwnerComponent) const = 0;
+
+	/**
+	 * Called by the engine when an external object wants to inject collision data into this simulation above
+	 * and beyond what is specified in the asset for the internal actors
+	 * Examples: Scene collision, collision for parents we are attached to
+	 * @param InData - Collisions to add to this simulation
+	 */
+	virtual void AddExternalCollisions(const FClothCollisionData& InData) = 0;
+
+	/**
+	 * Called by the engine when external collisions are no longer necessary or when they need to be updated
+	 * with some of the previous collisions removed. It is recommended in derived simulations to avoid freeing
+	 * any allocations regarding external collisions as it is likely more will be added soon after this call
+	 */
+	virtual void ClearExternalCollisions() = 0;
+
+	/**
+	 * Called by the engine to request data on all active collisions in a simulation. if bIncludeExternal is
+	 * true, derived implementations should add the asset collisions and any collisions added at runtime, if
+	 * false only the collisions from the asset should be considered
+	 * @param OutCollisions - Array to write collisions to
+	 * @param bIncludeExternal - Whether or not external collisions should be retrieved, or just asset collisions
+	 */
+	virtual void GetCollisions(FClothCollisionData& OutCollisions, bool bIncludeExternal = true) const = 0;
+
+private:
+
+public:
+
+	/**
+	 * Called by the engine once per physics tick to gather non-cycle stats (num sim verts etc.)
+	 * Will not be called if STATS is not defined, not required to be implemented if no need for stats.
+	 * STAT_NumCloths and STAT_NumClothVerts are already defined and can be incremented, more can be
+	 * added per simulation.
+	 * This can be called on any thread, so derived implementations should make sure the gather is
+	 * safe
+	 */
+	virtual void GatherStats() const
+	{}
+};
\ No newline at end of file
diff --git a/Engine/Source/Runtime/ClothingSystemRuntimeInterface/Public/ClothingSimulationContext.h b/Engine/Source/Runtime/ClothingSystemRuntimeInterface/Public/ClothingSimulationContext.h
new file mode 100644
--- /dev/null
+++ b/Engine/Source/Runtime/ClothingSystemRuntimeInterface/Public/ClothingSimulationContext.h
@@ -0,0 +1,53 @@
+// Copyright 1998-2019 Epic Games, Inc. All Rights Reserved.
+#pragma once
+
+#include "ClothingSystemRuntimeTypes.h"
+#include "Containers/Array.h"
+#include "Math/Transform.h"
+#include "Math/Vector.h"
+
+/** Empty interface, derived simulation modules define the contents of the context. */
+class CLOTHINGSYSTEMRUNTIMEINTERFACE_API IClothingSimulationContext
+{
+public:
+	IClothingSimulationContext();
+	virtual ~IClothingSimulationContext();
+};
+
+/** Base simlation data that just about every simulation would need. */
+class CLOTHINGSYSTEMRUNTIMEINTERFACE_API FClothingSimulationContextBase : public IClothingSimulationContext
+{
+public:
+
+	FClothingSimulationContextBase();
+	virtual ~FClothingSimulationContextBase();
+
+	// Delta for this tick
+	float DeltaSeconds;
+
+	// World space bone transforms of the owning component
+	TArray<FTransform> BoneTransforms;
+
+	// Component to world transform of the owning component
+	FTransform ComponentToWorld;
+
+	// The predicted LOD of the skeletal mesh component running the simulation
+	int32 PredictedLod;
+
+	// Wind velocity at the component location
+	FVector WindVelocity;
+
+	// Gravity extracted from the world
+	FVector WorldGravity;
+
+	// Wind adaption, a measure of how quickly to adapt to the wind speed
+	// when using the legacy wind calculation mode
+	float WindAdaption;
+
+	// Whether and how we should teleport the simulation this tick
+	EClothingTeleportMode TeleportMode;
+
+	// Scale for the max distance constraints of the simulation mesh
+	float MaxDistanceScale;
+};
+
diff --git a/Engine/Source/Runtime/ClothingSystemRuntimeInterface/Public/ClothingSimulationFactory.h b/Engine/Source/Runtime/ClothingSystemRuntimeInterface/Public/ClothingSimulationFactory.h
new file mode 100644
--- /dev/null
+++ b/Engine/Source/Runtime/ClothingSystemRuntimeInterface/Public/ClothingSimulationFactory.h
@@ -0,0 +1,62 @@
+// Copyright 1998-2019 Epic Games, Inc. All Rights Reserved.
+
+#pragma once
+
+#include "ClothingSimulationInterface.h"
+
+#include "UObject/Object.h"
+#include "Features/IModularFeature.h"
+
+#include "ClothingSimulationFactory.generated.h"
+
+class UClothingAssetBase;
+class UClothingSimulationInteractor;
+
+// An interface for a class that will provide default simulation factory classes
+// Used by modules wanting to override clothing simulation to provide their own implementation
+class CLOTHINGSYSTEMRUNTIMEINTERFACE_API IClothingSimulationFactoryClassProvider : public IModularFeature
+{
+public:
+
+	// The feature name to register against for providers
+	static const FName FeatureName;
+
+	// Called by the engine to get the default clothing simulation factory to use
+	// for skeletal mesh components (see USkeletalMeshComponent constructor).
+	// Returns Factory class for simulations or nullptr to disable clothing simulation
+	virtual UClass* GetDefaultSimulationFactoryClass() = 0;
+};
+
+// Any clothing simulation factory should derive from this interface object to interact with the engine
+UCLASS(Abstract)
+class CLOTHINGSYSTEMRUNTIMEINTERFACE_API UClothingSimulationFactory : public UObject
+{
+	GENERATED_BODY()
+
+public:
+
+	// Create a simulation object for a skeletal mesh to use (see IClothingSimulation)
+	virtual IClothingSimulation* CreateSimulation()
+	PURE_VIRTUAL(UClothingSimulationFactory::CreateSimulation, return nullptr;);
+
+	// Destroy a simulation object, guaranteed to be a pointer returned from CreateSimulation for this factory
+	virtual void DestroySimulation(IClothingSimulation* InSimulation)
+	PURE_VIRTUAL(UClothingSimulationFactory::DestroySimulation, );
+
+	// Given an asset, decide whether this factory can create a simulation to use the data inside
+	// (return false if data is invalid or missing in the case of custom data)
+	virtual bool SupportsAsset(UClothingAssetBase* InAsset)
+	PURE_VIRTUAL(UClothingSimulationFactory::SupportsAsset, return false;);
+
+	// Whether or not we provide an interactor object to manipulate the simulation at runtime.
+	// If true is returned then CreateInteractor *must* create a valid object to handle this
+	virtual bool SupportsRuntimeInteraction()
+	PURE_VIRTUAL(UClothingSimulationFactory::SupportsRuntimeInteraction, return false;);
+
+	// Creates the runtime interactor object for a clothing simulation. This object will
+	// receive events allowing it to write data to the simulation context in a safe manner
+	virtual UClothingSimulationInteractor* CreateInteractor()
+	{
+		return nullptr;
+	}
+};
diff --git a/Engine/Source/Runtime/ClothingSystemRuntimeInterface/Public/ClothingSimulationInterface.h b/Engine/Source/Runtime/ClothingSystemRuntimeInterface/Public/ClothingSimulationInterface.h
--- a/Engine/Source/Runtime/ClothingSystemRuntimeInterface/Public/ClothingSimulationInterface.h
+++ b/Engine/Source/Runtime/ClothingSystemRuntimeInterface/Public/ClothingSimulationInterface.h
@@ -1,138 +1,5 @@
 // Copyright 1998-2019 Epic Games, Inc. All Rights Reserved.
-
 #pragma once
 
-#include "ClothingSystemRuntimeTypes.h"
-
-class UClothingAssetBase;
-class USkeletalMeshComponent;
-
-// Empty interface, derived simulation modules define the contents of the context
-class CLOTHINGSYSTEMRUNTIMEINTERFACE_API IClothingSimulationContext
-{
-public:
-	virtual ~IClothingSimulationContext() {};
-};
-
-class CLOTHINGSYSTEMRUNTIMEINTERFACE_API IClothingSimulation
-{
-public:
-
-	virtual ~IClothingSimulation() {}
-
-	// The majority of the API for this class is protected. The required objects (skel meshes and the parallel task)
-	// are friends so they can use the functionality. For the most part the simulation is not designed to be used
-	// outside of the skeletal mesh component as it parallel simulation is tied to the skel mesh tick and
-	// dependents.
-	// Any method that is available in the public section below should consider that it may be called while
-	// the simulation is running.
-	// Currently only a const pointer is exposed externally to the skeletal mesh component, so barring that
-	// being cast away external callers can't call non-const methods. The debug skel mesh component exposes a
-	// mutable version for editor use
-protected:
-
-	friend class USkeletalMeshComponent;
-	friend class FParallelClothTask;
-
-	/**
-	 * Create an actor for this simulation from the data in InAsset
-	 * Simulation data for this actor should be written back to SimDataIndex in GetSimulationData
-	 * @param InOwnerComponent - The component requesting this actor
-	 * @param InAsset - The asset to create an actor from
-	 * @param SimDataIndex - the sim data index to use when doing the writeback in GetSimulationData
-	 */
-	virtual void CreateActor(USkeletalMeshComponent* InOwnerComponent, UClothingAssetBase* InAsset, int32 SimDataIndex) = 0;
-	
-	/** Create a new context, will not be filled, call FillContext before simulating with this context */
-	virtual IClothingSimulationContext* CreateContext() = 0;
-
-	/**
-	 * Fills an existing context for a single simulation step, called by the engine on the game thread prior to simulation 
-	 * @param InComponent - The component to fill the context for
-	 * @param InOutContext - The context to fill
-	 */
-	virtual void FillContext(USkeletalMeshComponent* InComponent, float InDeltaTime, IClothingSimulationContext* InOutContext) = 0;
-
-	/** Initialize the simulation, will be called before any Simulate calls */
-	virtual void Initialize() = 0;
-
-	/**
-	 * Shutdown the simulation, this should clear ALL resources as we no longer expect to
-	 * call any other function on this simulation without first calling Initialize again.
-	 */
-	virtual void Shutdown() = 0;
-
-	/**
-	 * Called by the engine to detect whether or not the simulation should run, essentially
-	 * are there any actors that need to simulate in this simulation
-	 */
-	virtual bool ShouldSimulate() const = 0;
-
-	/**
-	 * Run a single tick of the simulation.
-	 * The pointer InContext is guaranteed (when called by the engine) to be the context allocated in
-	 * CreateContext and can be assumed to be safely castable to any derived type allocated there.
-	 * New callers should make care to make sure only the correct context is ever passed through
-	 * @param InContext - The context to use during simulation, will have been filled in FillContext
-	 */
-	virtual void Simulate(IClothingSimulationContext* InContext) = 0;
-
-	/** Simulation should remove all of it's actors when next possible and free them */
-	virtual void DestroyActors() = 0;
-
-	/** 
-	 * Destroy a context object, engine will always pass a context created using CreateContext 
-	 * @param InContext - The context to destroy
-	 */
-	virtual void DestroyContext(IClothingSimulationContext* InContext) = 0;
-
-	/** 
-	 * Fill FClothSimulData map for the clothing simulation. Should fill a map pair per-actor 
-	 * @param OutData - The simulation data to write to
-	 * @param InOwnerComponent - the component that owns the simulation
-	 * @param InOverrideComponent - An override component if bound to a master pose component
-	 */
-	virtual void GetSimulationData(TMap<int32, FClothSimulData>& OutData, USkeletalMeshComponent* InOwnerComponent, USkinnedMeshComponent* InOverrideComponent) const = 0;
-
-	/** Get the bounds of the simulation mesh in local simulation space */
-	virtual FBoxSphereBounds GetBounds(const USkeletalMeshComponent* InOwnerComponent) const = 0;
-
-	/**
-	 * Called by the engine when an external object wants to inject collision data into this simulation above
-	 * and beyond what is specified in the asset for the internal actors
-	 * Examples: Scene collision, collision for parents we are attached to
-	 * @param InData - Collisions to add to this simulation
-	 */
-	virtual void AddExternalCollisions(const FClothCollisionData& InData) = 0;
-
-	/**
-	 * Called by the engine when external collisions are no longer necessary or when they need to be updated
-	 * with some of the previous collisions removed. It is recommended in derived simulations to avoid freeing
-	 * any allocations regarding external collisions as it is likely more will be added soon after this call
-	 */
-	virtual void ClearExternalCollisions() = 0;
-
-	/**
-	 * Called by the engine to request data on all active collisions in a simulation. if bIncludeExternal is
-	 * true, derived implementations should add the asset collisions and any collisions added at runtime, if
-	 * false only the collisions from the asset should be considered
-	 * @param OutCollisions - Array to write collisions to
-	 * @param bIncludeExternal - Whether or not external collisions should be retrieved, or just asset collisions
-	 */
-	virtual void GetCollisions(FClothCollisionData& OutCollisions, bool bIncludeExternal = true) const = 0;
-
-private:
-
-public:
-
-	/**
-	 * Called by the engine once per physics tick to gather non-cycle stats (num sim verts etc.)
-	 * Will not be called if STATS is not defined, not required to be implemented if no need for stats.
-	 * STAT_NumCloths and STAT_NumClothVerts are already defined and can be incremented, more can be
-	 * added per simulation.
-	 * This can be called on any thread, so derived implementations should make sure the gather is
-	 * safe
-	 */
-	virtual void GatherStats() const
-	{}
-};
\ No newline at end of file
+#include "ClothingSimulationContext.h"
+#include "ClothingSimulation.h"
diff --git a/Engine/Source/Runtime/ClothingSystemRuntimeInterface/Public/ClothingSystemRuntimeTypes.h b/Engine/Source/Runtime/ClothingSystemRuntimeInterface/Public/ClothingSystemRuntimeTypes.h
--- a/Engine/Source/Runtime/ClothingSystemRuntimeInterface/Public/ClothingSystemRuntimeTypes.h
+++ b/Engine/Source/Runtime/ClothingSystemRuntimeInterface/Public/ClothingSystemRuntimeTypes.h
@@ -3,8 +3,9 @@
 #pragma once
 
 #include "Containers/Array.h"
-#include "UObject/ObjectMacros.h"
-#include "ClothingSystemRuntimeTypes.generated.h"
+#include "Math/Transform.h"
+#include "Math/Vector.h"
+
 
 // Data produced by a clothing simulation
 struct FClothSimulData
@@ -38,77 +39,3 @@ enum class EClothingTeleportMode : uint8
 	TeleportAndReset
 };
 
-/** Data for a single sphere primitive in the clothing simulation. This can either be a 
- *  sphere on its own, or part of a capsule referenced by the indices in FClothCollisionPrim_Capsule
- */
-USTRUCT()
-struct FClothCollisionPrim_Sphere
-{
-	GENERATED_BODY()
-
-	UPROPERTY()
-	int32 BoneIndex;
-
-	UPROPERTY()
-	float Radius;
-
-	UPROPERTY()
-	FVector LocalPosition;
-};
-
-/** Data for a single connected sphere primitive. This should be configured after all spheres have
- *  been processed as they are really just indexing the existing spheres
- */
-USTRUCT()
-struct FClothCollisionPrim_SphereConnection
-{
-	GENERATED_BODY()
-
-	UPROPERTY()
-	int32 SphereIndices[2];
-};
-
-/**
- *	Data for a single convex element
- *	A convex is a collection of planes, in which the clothing will attempt to stay outside of the
- *	shape created by the planes combined.
- */
-USTRUCT()
-struct FClothCollisionPrim_Convex
-{
-	GENERATED_BODY()
-
-	FClothCollisionPrim_Convex()
-		: BoneIndex(INDEX_NONE)
-	{
-
-	}
-
-	UPROPERTY()
-	TArray<FPlane> Planes;
-
-	UPROPERTY()
-	int32 BoneIndex;
-};
-
-USTRUCT()
-struct CLOTHINGSYSTEMRUNTIMEINTERFACE_API FClothCollisionData
-{
-	GENERATED_BODY()
-
-	void Reset();
-
-	void Append(const FClothCollisionData& InOther);
-
-	// Sphere data
-	UPROPERTY(EditAnywhere, Category = Collison)
-	TArray<FClothCollisionPrim_Sphere> Spheres;
-
-	// Capsule data
-	UPROPERTY(EditAnywhere, Category = Collison)
-	TArray<FClothCollisionPrim_SphereConnection> SphereConnections;
-
-	// Convex Data
-	UPROPERTY(EditAnywhere, Category = Collison)
-	TArray<FClothCollisionPrim_Convex> Convexes;
-};
diff --git a/Engine/Source/Runtime/ClothingSystemRuntimeNv/ClothingSystemRuntimeNv.Build.cs b/Engine/Source/Runtime/ClothingSystemRuntimeNv/ClothingSystemRuntimeNv.Build.cs
new file mode 100644
--- /dev/null
+++ b/Engine/Source/Runtime/ClothingSystemRuntimeNv/ClothingSystemRuntimeNv.Build.cs
@@ -0,0 +1,33 @@
+// Copyright 1998-2019 Epic Games, Inc. All Rights Reserved.
+
+using UnrealBuildTool;
+
+public class ClothingSystemRuntimeNv : ModuleRules
+{
+	public ClothingSystemRuntimeNv(ReadOnlyTargetRules Target) : base(Target)
+	{
+		PrivateIncludePaths.Add("Runtime/ClothingSystemRuntimeNv/Private");
+        
+        SetupModulePhysicsSupport(Target);
+
+        PublicIncludePathModuleNames.Add("ClothingSystemRuntimeInterface");
+
+        PublicDependencyModuleNames.AddRange(
+            new string[] {
+                "Core",
+                "CoreUObject",
+                "ClothingSystemRuntimeInterface"
+            }
+        );
+
+        PrivateDependencyModuleNames.AddRange(
+            new string[]
+            {
+                "Engine",
+                "RenderCore",
+                "SlateCore",
+                "Slate"
+            }
+        );
+    }
+}
diff --git a/Engine/Source/Runtime/ClothingSystemRuntimeNv/Private/Assets/ClothingAsset.cpp b/Engine/Source/Runtime/ClothingSystemRuntimeNv/Private/Assets/ClothingAsset.cpp
new file mode 100644
--- /dev/null
+++ b/Engine/Source/Runtime/ClothingSystemRuntimeNv/Private/Assets/ClothingAsset.cpp
@@ -0,0 +1,1205 @@
+// Copyright 1998-2019 Epic Games, Inc. All Rights Reserved.
+
+#include "Assets/ClothingAsset.h"
+#include "PhysXPublic.h"
+
+#if WITH_EDITOR
+#include "Engine/SkeletalMesh.h"
+#include "MeshUtilities.h"
+#endif
+
+#include "Components/SkeletalMeshComponent.h"
+
+#include "PhysicsPublic.h"
+#include "UObject/UObjectIterator.h"
+#include "Modules/ModuleManager.h"
+#include "Widgets/Notifications/SNotificationList.h"
+#include "Framework/Notifications/NotificationManager.h"
+#include "ComponentReregisterContext.h"
+#include "UObject/AnimPhysObjectVersion.h"
+#include "Utils/ClothingMeshUtils.h"
+#include "Rendering/SkeletalMeshModel.h"
+#include "Rendering/SkeletalMeshRenderData.h"
+#include "ClothingSimulationInteractor.h"
+#include "Serialization/CustomVersion.h"
+
+DEFINE_LOG_CATEGORY(LogClothingAsset)
+
+#define LOCTEXT_NAMESPACE "ClothingAsset"
+
+
+// Custom serialization version for clothing assets
+struct FClothingAssetCustomVersion
+{
+	enum Type
+	{
+		// Before any version changes were made
+		BeforeCustomVersionWasAdded = 0,
+		// Added storage of vertex colors with sim data, for editor usage
+		AddVertexColorsToPhysicalMesh = 1,
+
+		// -----<new versions can be added above this line>-------------------------------------------------
+		VersionPlusOne,
+		LatestVersion = VersionPlusOne - 1
+	};
+
+	// The GUID for this custom version number
+	const static FGuid GUID;
+
+private:
+	FClothingAssetCustomVersion() {}
+};
+
+const FGuid FClothingAssetCustomVersion::GUID(0xFB680AF2, 0x59EF4BA3, 0xBAA819B5, 0x73C8443D);
+FCustomVersionRegistration GRegisterClothingAssetCustomVersion(FClothingAssetCustomVersion::GUID, FClothingAssetCustomVersion::LatestVersion, TEXT("ClothingAssetVer"));
+
+
+UClothingAsset::UClothingAsset(const FObjectInitializer& ObjectInitializer)
+	: Super(ObjectInitializer)
+	, ReferenceBoneIndex(0)
+	, CustomData(nullptr)
+{
+
+}
+
+void UClothingAsset::RefreshBoneMapping(USkeletalMesh* InSkelMesh)
+{
+	// No mesh, can't remap
+	if(!InSkelMesh)
+	{
+		return;
+	}
+
+	if(UsedBoneNames.Num() != UsedBoneIndices.Num())
+	{
+		UsedBoneIndices.Reset();
+		UsedBoneIndices.AddDefaulted(UsedBoneNames.Num());
+	}
+
+	// Repopulate the used indices.
+	for(int32 BoneNameIndex = 0; BoneNameIndex < UsedBoneNames.Num(); ++BoneNameIndex)
+	{
+		UsedBoneIndices[BoneNameIndex] = InSkelMesh->RefSkeleton.FindBoneIndex(UsedBoneNames[BoneNameIndex]);
+	}
+}
+
+#if WITH_EDITOR
+
+void LogAndToastClothingInfo(const FText& Error)
+{
+	FNotificationInfo Info(Error);
+	Info.ExpireDuration = 5.0f;
+	FSlateNotificationManager::Get().AddNotification(Info);
+
+	UE_LOG(LogClothingAsset, Warning, TEXT("%s"), *Error.ToString());
+}
+
+bool UClothingAsset::BindToSkeletalMesh(USkeletalMesh* InSkelMesh, int32 InMeshLodIndex, int32 InSectionIndex, int32 InAssetLodIndex, bool bCallPostEditChange)
+{
+	// If we've been added to the wrong mesh
+	if(InSkelMesh != GetOuter())
+	{
+		FText Error = FText::Format(LOCTEXT("Error_WrongMesh", "Failed to bind clothing asset {0} as the provided mesh is not the owner of this asset."), FText::FromString(GetName()));
+		LogAndToastClothingInfo(Error);
+
+		return false;
+	}
+
+	// If we don't have clothing data
+	if(!LodData.IsValidIndex(InAssetLodIndex))
+	{
+		FText Error = FText::Format(LOCTEXT("Error_NoClothingLod", "Failed to bind clothing asset {0} LOD{1} as LOD{2} does not exist."), FText::FromString(GetName()), FText::AsNumber(InAssetLodIndex), FText::AsNumber(InAssetLodIndex));
+		LogAndToastClothingInfo(Error);
+
+		return false;
+	}
+
+	// If we don't have a mesh
+	if(!InSkelMesh)
+	{
+		FText Error = FText::Format(LOCTEXT("Error_NoMesh", "Failed to bind clothing asset {0} as provided skel mesh does not exist."), FText::FromString(GetName()));
+		LogAndToastClothingInfo(Error);
+
+		return false;
+	}
+
+	// If the mesh LOD index is invalid
+	if(!InSkelMesh->GetImportedModel()->LODModels.IsValidIndex(InMeshLodIndex))
+	{
+		FText Error = FText::Format(LOCTEXT("Error_InvalidMeshLOD", "Failed to bind clothing asset {0} as mesh LOD{1} does not exist."), FText::FromString(GetName()), FText::AsNumber(InMeshLodIndex));
+		LogAndToastClothingInfo(Error);
+
+		return false;
+	}
+
+	const int32 NumMapEntries = LodMap.Num();
+	for(int MapIndex = 0; MapIndex < NumMapEntries; ++MapIndex)
+	{
+		const int32& MappedLod = LodMap[MapIndex];
+
+		if(MappedLod == InAssetLodIndex)
+		{
+			FText Error = FText::Format(LOCTEXT("Error_LodMapped", "Failed to bind clothing asset {0} LOD{1} as LOD{2} is already mapped to mesh LOD{3}."), FText::FromString(GetName()), FText::AsNumber(InAssetLodIndex), FText::AsNumber(InAssetLodIndex), FText::AsNumber(MapIndex));
+			LogAndToastClothingInfo(Error);
+
+			return false;
+		}
+	}
+
+	if(LodMap.IsValidIndex(InMeshLodIndex) && LodMap[InMeshLodIndex] != INDEX_NONE)
+	{
+		// Already mapped
+		return false;
+	}
+
+	BuildSelfCollisionData();
+	CalculateReferenceBoneIndex();
+
+	// Grab the clothing and skel lod data
+	FClothLODData& ClothLodData = LodData[InAssetLodIndex];
+	FSkeletalMeshLODModel& SkelLod = InSkelMesh->GetImportedModel()->LODModels[InMeshLodIndex];
+
+	FSkelMeshSection& OriginalSection = SkelLod.Sections[InSectionIndex];
+
+	// Data for mesh to mesh binding
+	TArray<FMeshToMeshVertData> MeshToMeshData;
+	TArray<FVector> RenderPositions;
+	TArray<FVector> RenderNormals;
+	TArray<FVector> RenderTangents;
+
+	RenderPositions.Reserve(OriginalSection.SoftVertices.Num());
+	RenderNormals.Reserve(OriginalSection.SoftVertices.Num());
+	RenderTangents.Reserve(OriginalSection.SoftVertices.Num());
+
+	// Original data to weight to the clothing simulation mesh
+	for(FSoftSkinVertex& UnrealVert : OriginalSection.SoftVertices)
+	{
+		RenderPositions.Add(UnrealVert.Position);
+		RenderNormals.Add(UnrealVert.TangentZ);
+		RenderTangents.Add(UnrealVert.TangentX);
+	}
+
+	TArrayView<uint32> IndexView(SkelLod.IndexBuffer);
+	IndexView.Slice(OriginalSection.BaseIndex, OriginalSection.NumTriangles * 3);
+
+	ClothingMeshUtils::ClothMeshDesc TargetMesh(RenderPositions, RenderNormals, IndexView);
+	ClothingMeshUtils::ClothMeshDesc SourceMesh(ClothLodData.PhysicalMeshData.Vertices, ClothLodData.PhysicalMeshData.Normals, ClothLodData.PhysicalMeshData.Indices);
+
+	ClothingMeshUtils::GenerateMeshToMeshSkinningData(MeshToMeshData, TargetMesh, &RenderTangents, SourceMesh);
+
+	if(MeshToMeshData.Num() == 0)
+	{
+		// Failed to generate skinning data, the function above will have notified
+		// with the cause of the failure, so just exit
+		return false;
+	}
+
+	// Calculate fixed verts
+	for(FMeshToMeshVertData& VertData : MeshToMeshData)
+	{
+		float TriangleDistanceMax = 0.0f;
+		TriangleDistanceMax += ClothLodData.PhysicalMeshData.MaxDistances[VertData.SourceMeshVertIndices[0]];
+		TriangleDistanceMax += ClothLodData.PhysicalMeshData.MaxDistances[VertData.SourceMeshVertIndices[1]];
+		TriangleDistanceMax += ClothLodData.PhysicalMeshData.MaxDistances[VertData.SourceMeshVertIndices[2]];
+
+		if(TriangleDistanceMax == 0.0f)
+		{
+			VertData.SourceMeshVertIndices[3] = 0xFFFF;
+		}
+	}
+
+	// We have to copy the bone map to verify we don't exceed the maximum while adding the clothing bones
+	TArray<FBoneIndexType> TempBoneMap = OriginalSection.BoneMap;
+
+	for(FName& BoneName : UsedBoneNames)
+	{
+		const int32 BoneIndex = InSkelMesh->RefSkeleton.FindBoneIndex(BoneName);
+
+		if(BoneIndex != INDEX_NONE)
+		{
+			TempBoneMap.AddUnique(BoneIndex);
+		}
+	}
+	
+	// Verify number of bones against current capabilities
+	if(TempBoneMap.Num() > FGPUBaseSkinVertexFactory::GetMaxGPUSkinBones())
+	{
+		// Failed to apply as we've exceeded the number of bones we can skin
+		FText Error = FText::Format(LOCTEXT("Error_TooManyBones", "Failed to bind clothing asset {0} LOD{1} as this causes the section to require {2} bones. The maximum per section is currently {3}."), FText::FromString(GetName()), FText::AsNumber(InAssetLodIndex), FText::AsNumber(TempBoneMap.Num()), FText::AsNumber(FGPUBaseSkinVertexFactory::GetMaxGPUSkinBones()));
+		LogAndToastClothingInfo(Error);
+
+		return false;
+	}
+
+	// After verifying copy the new bone map to the section
+	OriginalSection.BoneMap = TempBoneMap;
+
+	// Array of re-import contexts for components using this mesh
+	TIndirectArray<FComponentReregisterContext> ComponentContexts;
+	for (TObjectIterator<USkeletalMeshComponent> It; It; ++It)
+	{
+		USkeletalMeshComponent* Component = *It;
+		if (Component && !Component->IsTemplate() && Component->SkeletalMesh == InSkelMesh)
+		{
+			ComponentContexts.Add(new FComponentReregisterContext(Component));
+		}
+	}
+
+	// Ready to apply the changes
+	InSkelMesh->PreEditChange(nullptr);
+
+	// calculate LOD verts before adding our new section
+	uint32 NumLodVertices = 0;
+	for(const FSkelMeshSection& CurSection : SkelLod.Sections)
+	{
+		NumLodVertices += CurSection.GetNumVertices();
+	}
+
+	// Set the asset index, used during rendering to pick the correct sim mesh buffer
+	int32 AssetIndex = INDEX_NONE;
+	check(InSkelMesh->MeshClothingAssets.Find(this, AssetIndex));
+	OriginalSection.CorrespondClothAssetIndex = AssetIndex;
+
+	// sim properties
+	OriginalSection.ClothMappingData = MeshToMeshData;
+	OriginalSection.ClothingData.AssetGuid = AssetGuid;
+	OriginalSection.ClothingData.AssetLodIndex = InAssetLodIndex;
+
+	bool bRequireBoneChange = false;
+	for(FBoneIndexType& BoneIndex : OriginalSection.BoneMap)
+	{
+		if(!SkelLod.RequiredBones.Contains(BoneIndex))
+		{
+			bRequireBoneChange = true;
+			if(InSkelMesh->RefSkeleton.IsValidIndex(BoneIndex))
+			{
+				SkelLod.RequiredBones.Add(BoneIndex);
+				SkelLod.ActiveBoneIndices.AddUnique(BoneIndex);
+			}
+		}
+	}
+
+	if(bRequireBoneChange)
+	{
+		SkelLod.RequiredBones.Sort();
+		InSkelMesh->RefSkeleton.EnsureParentsExistAndSort(SkelLod.ActiveBoneIndices);		
+	}
+
+	if(CustomData)
+	{
+		CustomData->BindToSkeletalMesh(InSkelMesh, InMeshLodIndex, InSectionIndex, InAssetLodIndex);
+	}
+
+	// Make sure the LOD map is always big enough for the asset to use.
+	// This shouldn't grow to an unwieldy size but maybe consider compacting later.
+	while(LodMap.Num() - 1 < InMeshLodIndex)
+	{
+		LodMap.Add(INDEX_NONE);
+	}
+
+	LodMap[InMeshLodIndex] = InAssetLodIndex;
+
+	if (bCallPostEditChange)
+	{
+	InSkelMesh->PostEditChange();
+	}
+
+	return true;
+
+	// ComponentContexts goes out of scope, causing components to be re-registered
+}
+
+void UClothingAsset::UnbindFromSkeletalMesh(USkeletalMesh* InSkelMesh)
+{
+	if(FSkeletalMeshModel* Mesh = InSkelMesh->GetImportedModel())
+	{
+		const int32 NumLods = Mesh->LODModels.Num();
+
+		for(int32 LodIndex = 0; LodIndex < NumLods; ++LodIndex)
+		{
+			UnbindFromSkeletalMesh(InSkelMesh, LodIndex);
+		}
+	}
+}
+
+void UClothingAsset::UnbindFromSkeletalMesh(USkeletalMesh* InSkelMesh, int32 InMeshLodIndex)
+{
+	bool bChangedMesh = false;
+
+	// Find the chunk(s) we created
+	if(FSkeletalMeshModel* Mesh = InSkelMesh->GetImportedModel())
+	{
+		if(!Mesh->LODModels.IsValidIndex(InMeshLodIndex))
+		{
+			FText Error = FText::Format(LOCTEXT("Error_UnbindNoMeshLod", "Failed to remove clothing asset {0} from mesh LOD{1} as that LOD doesn't exist."), FText::FromString(GetName()), FText::AsNumber(InMeshLodIndex));
+			LogAndToastClothingInfo(Error);
+
+			return;
+		}
+
+		FSkeletalMeshLODModel& LodModel = Mesh->LODModels[InMeshLodIndex];
+
+		for(int32 SectionIdx = LodModel.Sections.Num() - 1; SectionIdx >= 0; --SectionIdx)
+		{
+			FSkelMeshSection& Section = LodModel.Sections[SectionIdx];
+			if(Section.HasClothingData() && Section.ClothingData.AssetGuid == AssetGuid)
+			{
+				if(!bChangedMesh)
+				{
+					InSkelMesh->PreEditChange(nullptr);
+				}
+				ClothingAssetUtils::ClearSectionClothingData(Section);
+				bChangedMesh = true;
+			}
+		}
+
+		// Clear the LOD map entry for this asset LOD, after a unbind we must be able to bind any asset
+		if (LodMap.IsValidIndex(InMeshLodIndex))
+		{
+			LodMap[InMeshLodIndex] = INDEX_NONE;
+			bChangedMesh = true;
+		}
+	}
+
+	// If the mesh changed we need to re-register any components that use it to reflect the changes
+	if(bChangedMesh)
+	{
+		InSkelMesh->PostEditChange();
+
+		for(TObjectIterator<USkeletalMeshComponent> It; It; ++It)
+		{
+			USkeletalMeshComponent* MeshComponent = *It;
+			if(MeshComponent &&
+			   !MeshComponent->IsTemplate() &&
+			   MeshComponent->SkeletalMesh == InSkelMesh)
+			{
+				MeshComponent->ReregisterComponent();
+			}
+		}
+	}
+}
+
+void UClothingAsset::InvalidateCachedData()
+{
+	for(FClothLODData& CurrentLodData : LodData)
+	{
+		// Recalculate inverse masses for the physical mesh particles
+		FClothPhysicalMeshData& PhysMesh = CurrentLodData.PhysicalMeshData;
+
+		check(PhysMesh.Indices.Num() % 3 == 0);
+
+		TArray<float>& InvMasses = PhysMesh.InverseMasses;
+
+		const int32 NumVerts = PhysMesh.Vertices.Num();
+		InvMasses.Empty(NumVerts);
+		InvMasses.AddZeroed(NumVerts);
+
+		for(int32 TriBaseIndex = 0; TriBaseIndex < PhysMesh.Indices.Num(); TriBaseIndex += 3)
+		{
+			const int32 Index0 = PhysMesh.Indices[TriBaseIndex];
+			const int32 Index1 = PhysMesh.Indices[TriBaseIndex + 1];
+			const int32 Index2 = PhysMesh.Indices[TriBaseIndex + 2];
+
+			const FVector AB = PhysMesh.Vertices[Index1] - PhysMesh.Vertices[Index0];
+			const FVector AC = PhysMesh.Vertices[Index2] - PhysMesh.Vertices[Index0];
+			const float TriArea = FVector::CrossProduct(AB, AC).Size();
+
+			InvMasses[Index0] += TriArea;
+			InvMasses[Index1] += TriArea;
+			InvMasses[Index2] += TriArea;
+		}
+
+		bool bHasMaxDistance = PhysMesh.MaxDistances.Num() > 0;
+		PhysMesh.NumFixedVerts = 0;
+		
+		if(bHasMaxDistance)
+		{
+		float MassSum = 0.0f;
+		for(int32 CurrVertIndex = 0; CurrVertIndex < NumVerts; ++CurrVertIndex)
+		{
+			float& InvMass = InvMasses[CurrVertIndex];
+			const float& MaxDistance = PhysMesh.MaxDistances[CurrVertIndex];
+
+			if(MaxDistance < SMALL_NUMBER)
+			{
+				InvMass = 0.0f;
+				++PhysMesh.NumFixedVerts;
+			}
+			else
+			{
+				MassSum += InvMass;
+			}
+		}
+
+		if(MassSum > 0.0f)
+		{
+			const float MassScale = (float)(NumVerts - PhysMesh.NumFixedVerts) / MassSum;	
+
+			for(float& InvMass : InvMasses)
+			{
+				if(InvMass != 0.0f)
+				{
+					InvMass *= MassScale;
+					InvMass = 1.0f / InvMass;
+				}
+			}
+		}
+		}
+		else
+		{
+			for(int32 CurrVertIndex = 0; CurrVertIndex < NumVerts; ++CurrVertIndex)
+			{
+				InvMasses[CurrVertIndex] = 0.0f;
+			}
+
+			PhysMesh.NumFixedVerts = NumVerts;
+		}
+
+		// Calculate number of influences per vertex
+		for(int32 VertIndex = 0; VertIndex < NumVerts; ++VertIndex)
+		{
+			FClothVertBoneData& BoneData = PhysMesh.BoneData[VertIndex];
+			const uint16* BoneIndices = BoneData.BoneIndices;
+			const float* BoneWeights = BoneData.BoneWeights;
+
+			BoneData.NumInfluences = MAX_TOTAL_INFLUENCES;
+
+			int32 NumInfluences = 0;
+			for(int32 InfluenceIndex = 0; InfluenceIndex < MAX_TOTAL_INFLUENCES; ++InfluenceIndex)
+			{
+				if(BoneWeights[InfluenceIndex] == 0.0f || BoneIndices[InfluenceIndex] == INDEX_NONE)
+				{
+					BoneData.NumInfluences = NumInfluences;
+					break;
+				}
+				++NumInfluences;
+			}
+		}
+	}
+}
+
+void UClothingAsset::BuildLodTransitionData()
+{
+	const int32 NumLods = LodData.Num();
+	for(int32 LodIndex = 0; LodIndex < NumLods; ++LodIndex)
+	{
+		const bool bHasPrevLod = LodIndex > 0;
+		const bool bHasNextLod = LodIndex < NumLods - 1;
+
+		FClothLODData& CurrentLod = LodData[LodIndex];
+		FClothPhysicalMeshData& CurrentPhysMesh = CurrentLod.PhysicalMeshData;
+
+		FClothLODData* PrevLod = bHasPrevLod ? &LodData[LodIndex - 1] : nullptr;
+		FClothLODData* NextLod = bHasNextLod ? &LodData[LodIndex + 1] : nullptr;
+
+		const int32 CurrentLodNumVerts = CurrentPhysMesh.Vertices.Num();
+
+		ClothingMeshUtils::ClothMeshDesc CurrentMeshDesc(CurrentPhysMesh.Vertices, CurrentPhysMesh.Normals, CurrentPhysMesh.Indices);
+
+		if(PrevLod)
+		{
+			FClothPhysicalMeshData& PrevPhysMesh = PrevLod->PhysicalMeshData;
+
+			CurrentLod.TransitionUpSkinData.Empty(CurrentLodNumVerts);
+
+			ClothingMeshUtils::ClothMeshDesc PrevMeshDesc(PrevPhysMesh.Vertices, PrevPhysMesh.Normals, PrevPhysMesh.Indices);
+
+			ClothingMeshUtils::GenerateMeshToMeshSkinningData(CurrentLod.TransitionUpSkinData, CurrentMeshDesc, nullptr, PrevMeshDesc);
+		}
+
+		if(NextLod)
+		{
+			FClothPhysicalMeshData& NextPhysMesh = NextLod->PhysicalMeshData;
+
+			CurrentLod.TransitionDownSkinData.Empty(CurrentLodNumVerts);
+
+			ClothingMeshUtils::ClothMeshDesc NextMeshDesc(NextPhysMesh.Vertices, NextPhysMesh.Normals, NextPhysMesh.Indices);
+
+			ClothingMeshUtils::GenerateMeshToMeshSkinningData(CurrentLod.TransitionDownSkinData, CurrentMeshDesc, nullptr, NextMeshDesc);
+		}
+	}
+}
+
+void UClothingAsset::ApplyParameterMasks()
+{
+	for(FClothLODData& Lod : LodData)
+	{
+		// First zero out the parameters, otherwise disabled masks might hang around
+		Lod.PhysicalMeshData.ClearParticleParameters();
+
+		for(FClothParameterMask_PhysMesh& Mask : Lod.ParameterMasks)
+		{
+			// Only apply enabled masks
+			if(!Mask.bEnabled)
+			{
+				continue;
+			}
+
+			TArray<float>* TargetArray = nullptr;
+
+			switch(Mask.CurrentTarget)
+			{
+				case MaskTarget_PhysMesh::BackstopDistance:
+					TargetArray = &Lod.PhysicalMeshData.BackstopDistances;
+					break;
+				case MaskTarget_PhysMesh::BackstopRadius:
+					TargetArray = &Lod.PhysicalMeshData.BackstopRadiuses;
+					break;
+				case MaskTarget_PhysMesh::MaxDistance:
+					TargetArray = &Lod.PhysicalMeshData.MaxDistances;
+					break;
+				case MaskTarget_PhysMesh::AnimDriveMultiplier:
+					TargetArray = &Lod.PhysicalMeshData.AnimDriveMultipliers;
+					break;
+				default:
+					break;
+			}
+
+			if(TargetArray)
+			{
+				*TargetArray = Mask.GetValueArray();
+			}
+		}
+	}
+
+	InvalidateCachedData();
+}
+
+bool UClothingAsset::IsValidLod(int32 InLodIndex)
+{
+	return LodData.IsValidIndex(InLodIndex);
+}
+
+int32 UClothingAsset::GetNumLods()
+{
+	return LodData.Num();
+}
+
+#endif
+
+void UClothingAsset::BuildSelfCollisionData()
+{
+	if(!ClothConfig.HasSelfCollision())
+	{
+		// No self collision, can't generate data
+		return;
+	}
+
+	// can't pass through the network of other spheres.
+	const float SCRadius = ClothConfig.SelfCollisionRadius * ClothConfig.SelfCollisionCullScale;
+	const float SCRadiusSq = SCRadius * SCRadius;
+
+	for(FClothLODData& Lod : LodData)
+	{
+		FClothPhysicalMeshData& PhysMesh = Lod.PhysicalMeshData;
+		
+		// Start with the full set
+		const int32 NumVerts = PhysMesh.Vertices.Num();
+		PhysMesh.SelfCollisionIndices.Reset();
+		for(int32 Index = 0; Index < NumVerts; ++Index)
+		{
+			PhysMesh.SelfCollisionIndices.Add(Index);
+		}
+
+		// Strip any verts that are fixed from the beginning of the array so we always
+		// start with a valid movable vertex
+		int32 PrepassIndex = PhysMesh.SelfCollisionIndices[0];
+		while(PhysMesh.MaxDistances[PrepassIndex] < SMALL_NUMBER)
+		{
+			PhysMesh.SelfCollisionIndices.RemoveAt(0);
+
+			if(PhysMesh.SelfCollisionIndices.Num() > 0)
+			{
+				// Take index again
+				PrepassIndex = PhysMesh.SelfCollisionIndices[0];
+			}
+			else
+			{
+				// We've cleared out the array, bail as we have no movable verts
+				break;
+			}
+		}
+
+		// Now start aggresively culling verts that are near others that we have accepted
+		for(int32 Vert0Itr = 0; Vert0Itr < PhysMesh.SelfCollisionIndices.Num(); ++Vert0Itr)
+		{
+			uint32 V0Index = PhysMesh.SelfCollisionIndices[Vert0Itr];
+			const FVector& V0Pos = PhysMesh.Vertices[V0Index];
+			float MaxDist0 = PhysMesh.MaxDistances[V0Index];
+
+			// Start one after our current V0, we've done the other checks
+			for(int32 Vert1Itr = Vert0Itr + 1; Vert1Itr < PhysMesh.SelfCollisionIndices.Num(); )
+			{
+				uint32 V1Index = PhysMesh.SelfCollisionIndices[Vert1Itr];
+				const FVector& V1Pos = PhysMesh.Vertices[V1Index];
+
+				float V0ToV1DistSq = (V1Pos - V0Pos).SizeSquared();
+				float MaxDist1 = PhysMesh.MaxDistances[V1Index];
+
+				if(V0ToV1DistSq < SCRadiusSq || MaxDist1 < SMALL_NUMBER)
+				{
+					// Too close, remove it
+					PhysMesh.SelfCollisionIndices.RemoveAt(Vert1Itr);
+					continue;
+				}
+				else
+				{
+					// Move to next if we didn't remove
+					++Vert1Itr;
+				}
+			}
+		}
+	}
+}
+
+void UClothingAsset::PostLoad()
+{
+	Super::PostLoad();
+
+	BuildSelfCollisionData();
+
+#if WITH_EDITORONLY_DATA
+	CalculateReferenceBoneIndex();
+#endif
+
+	int32 CustomVersion = GetLinkerCustomVersion(FAnimPhysObjectVersion::GUID);
+
+	if(CustomVersion < FAnimPhysObjectVersion::AddedClothingMaskWorkflow)
+	{
+#if WITH_EDITORONLY_DATA
+		// Convert current parameters to masks
+		for(FClothLODData& Lod : LodData)
+		{
+			FClothPhysicalMeshData& PhysMesh = Lod.PhysicalMeshData;
+
+			// Didn't do anything previously - clear out incase there's something in there
+			// so we can use it correctly now.
+			Lod.ParameterMasks.Reset(3);
+
+			// Max distances (Always present)
+			Lod.ParameterMasks.AddDefaulted();
+			FClothParameterMask_PhysMesh& MaxDistanceMask = Lod.ParameterMasks.Last();
+			MaxDistanceMask.CopyFromPhysMesh(PhysMesh, MaskTarget_PhysMesh::MaxDistance);
+			MaxDistanceMask.bEnabled = true;
+
+			// Following params are only added if necessary, if we don't have any backstop
+			// radii then there's no backstops.
+			if(PhysMesh.BackstopRadiuses.FindByPredicate([](const float& A) {return A != 0.0f; }))
+			{
+				// Backstop radii
+				Lod.ParameterMasks.AddDefaulted();
+				FClothParameterMask_PhysMesh& BackstopRadiusMask = Lod.ParameterMasks.Last();
+				BackstopRadiusMask.CopyFromPhysMesh(PhysMesh, MaskTarget_PhysMesh::BackstopRadius);
+				BackstopRadiusMask.bEnabled = true;
+
+				// Backstop distances
+				Lod.ParameterMasks.AddDefaulted();
+				FClothParameterMask_PhysMesh& BackstopDistanceMask = Lod.ParameterMasks.Last();
+				BackstopDistanceMask.CopyFromPhysMesh(PhysMesh, MaskTarget_PhysMesh::BackstopDistance);
+				BackstopDistanceMask.bEnabled = true;
+			}
+			
+		}
+#endif
+
+		// Make sure we're transactional
+		SetFlags(RF_Transactional);
+	}
+
+#if WITH_EDITORONLY_DATA
+	// Fix content imported before we kept vertex colors
+	if(GetLinkerCustomVersion(FClothingAssetCustomVersion::GUID) < FClothingAssetCustomVersion::AddVertexColorsToPhysicalMesh)
+	{
+		for (FClothLODData& Lod : LodData)
+		{
+			const int32 NumVerts = Lod.PhysicalMeshData.Vertices.Num(); // number of verts
+
+			Lod.PhysicalMeshData.VertexColors.Reset();
+			Lod.PhysicalMeshData.VertexColors.AddUninitialized(NumVerts);
+			for (int32 VertIdx = 0; VertIdx < NumVerts; VertIdx++)
+			{
+				Lod.PhysicalMeshData.VertexColors[VertIdx] = FColor::White;
+			}
+		}
+	}
+#endif // WITH_EDITORONLY_DATA
+
+#if WITH_EDITOR
+	if(CustomVersion < FAnimPhysObjectVersion::CacheClothMeshInfluences)
+	{
+		// Rebuild data cache
+		InvalidateCachedData();
+	}
+#endif
+
+}
+
+void UClothingAsset::CalculateReferenceBoneIndex()
+{
+	// Starts at root
+	ReferenceBoneIndex = 0;
+
+	// Find the root bone for this clothing asset (common bone for all used bones)
+	typedef TArray<int32> BoneIndexArray;
+
+	// List of valid paths to the root bone from each weighted bone
+	TArray<BoneIndexArray> PathsToRoot;
+	
+	USkeletalMesh* OwnerMesh = Cast<USkeletalMesh>(GetOuter());
+
+	if(OwnerMesh)
+	{
+		FReferenceSkeleton& RefSkel = OwnerMesh->RefSkeleton;
+		// First build a list per used bone for it's path to root
+		const int32 NumUsedBones = UsedBoneIndices.Num();
+
+		// List of actually weighted (not just used) bones
+		TArray<int32> WeightedBones;
+
+		for(FClothLODData& CurLod : LodData)
+		{
+			FClothPhysicalMeshData& MeshData = CurLod.PhysicalMeshData;
+
+			for(FClothVertBoneData& VertBoneData : MeshData.BoneData)
+			{
+				for(int32 InfluenceIndex = 0; InfluenceIndex < MAX_TOTAL_INFLUENCES; ++InfluenceIndex)
+				{
+					if(VertBoneData.BoneWeights[InfluenceIndex] > SMALL_NUMBER)
+					{
+						WeightedBones.AddUnique(VertBoneData.BoneIndices[InfluenceIndex]);
+					}
+					else
+					{
+						// Hit the last weight (they're sorted)
+						break;
+					}
+				}
+			}
+		}
+
+		const int32 NumWeightedBones = WeightedBones.Num();
+		PathsToRoot.Reserve(NumWeightedBones);
+		
+		// Compute paths to the root bone
+		for(int32 WeightedBoneIndex = 0; WeightedBoneIndex < NumWeightedBones; ++WeightedBoneIndex)
+		{
+			PathsToRoot.AddDefaulted();
+			BoneIndexArray& Path = PathsToRoot.Last();
+			
+			int32 CurrentBone = WeightedBones[WeightedBoneIndex];
+			Path.Add(CurrentBone);
+			
+			while(CurrentBone != 0 && CurrentBone != INDEX_NONE)
+			{
+				CurrentBone = RefSkel.GetParentIndex(CurrentBone);
+				Path.Add(CurrentBone);
+			}
+		}
+
+		// Paths are from leaf->root, we want the other way
+		for(BoneIndexArray& Path : PathsToRoot)
+		{
+			Algo::Reverse(Path);
+		}
+
+		// Verify the last common bone in all paths as the root of the sim space
+		const int32 NumPaths = PathsToRoot.Num();
+		if(NumPaths > 0)
+		{
+			BoneIndexArray& FirstPath = PathsToRoot[0];
+		
+			const int32 FirstPathSize = FirstPath.Num();
+			for(int32 PathEntryIndex = 0; PathEntryIndex < FirstPathSize; ++PathEntryIndex)
+			{
+				const int32 CurrentQueryIndex = FirstPath[PathEntryIndex];
+				bool bValidRoot = true;
+
+				for(int32 PathIndex = 1; PathIndex < NumPaths; ++PathIndex)
+				{
+					if(!PathsToRoot[PathIndex].Contains(CurrentQueryIndex))
+					{
+						bValidRoot = false;
+						break;
+					}
+				}
+
+				if(bValidRoot)
+				{
+					ReferenceBoneIndex = CurrentQueryIndex;
+				}
+				else
+				{
+					// Once we fail to find a valid root we're done.
+					break;
+				}
+			}
+		}
+		else
+		{
+			// Just use root
+			ReferenceBoneIndex = 0;
+		}
+	}
+}
+
+#if WITH_EDITOR
+
+void UClothingAsset::PostEditChangeChainProperty(FPropertyChangedChainEvent& InEvent)
+{
+	bool bReregisterComponents = false;
+
+	if(InEvent.ChangeType != EPropertyChangeType::Interactive)
+	{
+		if(InEvent.Property->GetFName() == GET_MEMBER_NAME_CHECKED(FClothConfig, SelfCollisionRadius) ||
+			InEvent.Property->GetFName() == GET_MEMBER_NAME_CHECKED(FClothConfig, SelfCollisionCullScale))
+		{
+			BuildSelfCollisionData();
+			bReregisterComponents = true;
+		}
+		else if(InEvent.Property->GetFName() == GET_MEMBER_NAME_CHECKED(UClothingAsset, PhysicsAsset))
+		{
+			bReregisterComponents = true;
+		}
+		else
+		{
+			// Other properties just require a config refresh
+			ForEachInteractorUsingClothing([](UClothingSimulationInteractor* InInteractor)
+			{
+				if(InInteractor)
+				{
+					InInteractor->ClothConfigUpdated();
+				}
+			});
+		}
+	}
+
+	if(bReregisterComponents)
+	{
+		ReregisterComponentsUsingClothing();
+	}
+}
+
+void UClothingAsset::ReregisterComponentsUsingClothing()
+{
+	if(USkeletalMesh* OwnerMesh = Cast<USkeletalMesh>(GetOuter()))
+	{
+		for(TObjectIterator<USkeletalMeshComponent> It; It; ++It)
+		{
+			if(USkeletalMeshComponent* Component = *It)
+			{
+				if(Component->SkeletalMesh == OwnerMesh)
+				{
+					FComponentReregisterContext Context(Component);
+				}
+			}
+		}
+	}
+}
+
+void UClothingAsset::ForEachInteractorUsingClothing(TFunction<void(UClothingSimulationInteractor*)> Func)
+{
+	if(USkeletalMesh* OwnerMesh = Cast<USkeletalMesh>(GetOuter()))
+	{
+		for(TObjectIterator<USkeletalMeshComponent> It; It; ++It)
+		{
+			if(USkeletalMeshComponent* Component = *It)
+			{
+				if(Component->SkeletalMesh == OwnerMesh)
+				{
+					UClothingSimulationInteractor* CurInteractor = Component->GetClothingSimulationInteractor();
+
+					if(CurInteractor)
+					{
+						Func(CurInteractor);
+					}
+				}
+			}
+		}
+	}
+}
+
+#endif
+
+void UClothingAsset::Serialize(FArchive& Ar)
+{
+	Super::Serialize(Ar);
+	Ar.UsingCustomVersion(FAnimPhysObjectVersion::GUID);
+	Ar.UsingCustomVersion(FClothingAssetCustomVersion::GUID);
+}
+
+void ClothingAssetUtils::GetMeshClothingAssetBindings(USkeletalMesh* InSkelMesh, TArray<FClothingAssetMeshBinding>& OutBindings)
+{
+	OutBindings.Empty();
+
+	if(!InSkelMesh)
+	{
+		return;
+	}
+
+	if(FSkeletalMeshRenderData* Resource = InSkelMesh->GetResourceForRendering())
+	{
+		const int32 NumLods = Resource->LODRenderData.Num();
+
+		for(int32 LodIndex = 0; LodIndex < NumLods; ++LodIndex)
+		{
+			TArray<FClothingAssetMeshBinding> LodBindings;
+			GetMeshClothingAssetBindings(InSkelMesh, LodBindings, LodIndex);
+
+			OutBindings.Append(LodBindings);
+		}
+	}
+}
+
+void ClothingAssetUtils::GetMeshClothingAssetBindings(USkeletalMesh* InSkelMesh, TArray<FClothingAssetMeshBinding>& OutBindings, int32 InLodIndex)
+{
+	OutBindings.Empty();
+
+	if(!InSkelMesh)
+	{
+		return;
+	}
+
+	if(FSkeletalMeshRenderData* Resource = InSkelMesh->GetResourceForRendering())
+	{
+		if(Resource->LODRenderData.IsValidIndex(InLodIndex))
+		{
+			FSkeletalMeshLODRenderData& LodData = Resource->LODRenderData[InLodIndex];
+
+			const int32 NumSections = LodData.RenderSections.Num();
+
+			for(int32 SectionIndex = 0; SectionIndex < NumSections; ++SectionIndex)
+			{
+				FSkelMeshRenderSection& Section = LodData.RenderSections[SectionIndex];
+
+				if(Section.HasClothingData())
+				{
+					UClothingAsset* SectionAsset = Cast<UClothingAsset>(InSkelMesh->GetSectionClothingAsset(InLodIndex, SectionIndex));
+
+					if(SectionAsset)
+					{
+						// This is the original section of a clothing section pair
+						OutBindings.AddDefaulted();
+						FClothingAssetMeshBinding& Binding = OutBindings.Last();
+
+						Binding.Asset = SectionAsset;
+						Binding.LODIndex = InLodIndex;
+						Binding.SectionIndex = SectionIndex;
+						Binding.AssetInternalLodIndex = Section.ClothingData.AssetLodIndex;
+					}
+				}
+			}
+		}
+	}
+}
+
+#if WITH_EDITOR
+void ClothingAssetUtils::ClearSectionClothingData(FSkelMeshSection& InSection)
+{
+	InSection.ClothingData.AssetGuid = FGuid();
+	InSection.ClothingData.AssetLodIndex = INDEX_NONE;
+	InSection.CorrespondClothAssetIndex = INDEX_NONE;
+
+	InSection.ClothMappingData.Empty();
+}
+#endif
+
+bool FClothConfig::HasSelfCollision() const
+{
+	return SelfCollisionRadius > 0.0f && SelfCollisionStiffness > 0.0f;
+}
+
+void FClothPhysicalMeshData::Reset(const int32 InNumVerts)
+{
+	Vertices.Reset();
+	Normals.Reset();
+#if WITH_EDITORONLY_DATA
+	VertexColors.Reset();
+#endif // #if WITH_EDITORONLY_DATA
+	MaxDistances.Reset();
+	BackstopDistances.Reset();
+	BackstopRadiuses.Reset();
+	InverseMasses.Reset();
+	BoneData.Reset();
+
+	Vertices.AddDefaulted(InNumVerts);
+	Normals.AddDefaulted(InNumVerts);
+#if WITH_EDITORONLY_DATA
+	VertexColors.AddDefaulted(InNumVerts);
+#endif //#if WITH_EDITORONLY_DATA
+	MaxDistances.AddDefaulted(InNumVerts);
+	BackstopDistances.AddDefaulted(InNumVerts);
+	BackstopRadiuses.AddDefaulted(InNumVerts);
+	InverseMasses.AddDefaulted(InNumVerts);
+	BoneData.AddDefaulted(InNumVerts);
+
+	MaxBoneWeights = 0;
+	NumFixedVerts = 0;
+}
+
+void FClothPhysicalMeshData::ClearParticleParameters()
+{
+	// Max distances must be present, so fill to zero on clear so we still have valid mesh data.
+	const int32 NumVerts = Vertices.Num();
+	MaxDistances.Reset(NumVerts);
+	MaxDistances.AddZeroed(NumVerts);
+
+	// Just clear optional properties
+	BackstopDistances.Empty();
+	BackstopRadiuses.Empty();
+	AnimDriveMultipliers.Empty();
+}
+
+bool FClothPhysicalMeshData::HasBackStops() const
+{
+	const int32 NumBackStopDistances = BackstopDistances.Num();
+	return NumBackStopDistances > 0 && NumBackStopDistances == BackstopRadiuses.Num();
+}
+
+bool FClothPhysicalMeshData::HasAnimDrive() const
+{
+	return AnimDriveMultipliers.Num() > 0;
+}
+
+void FClothParameterMask_PhysMesh::Initialize(const FClothPhysicalMeshData& InMeshData)
+{
+	const int32 NumVerts = InMeshData.Vertices.Num();
+
+	// Set up value array
+	Values.Reset(NumVerts);
+	Values.AddZeroed(NumVerts);
+
+	bEnabled = false;
+}
+
+void FClothParameterMask_PhysMesh::CopyFromPhysMesh(const FClothPhysicalMeshData& InMeshData, MaskTarget_PhysMesh InTarget)
+{
+	// Presize value arrays
+	Initialize(InMeshData);
+
+	// Set our target
+	CurrentTarget = InTarget;
+
+	// Copy the actual parameter data
+	switch(InTarget)
+	{
+		case MaskTarget_PhysMesh::BackstopDistance:
+			Values = InMeshData.BackstopDistances;
+			break;
+		case MaskTarget_PhysMesh::BackstopRadius:
+			Values = InMeshData.BackstopRadiuses;
+			break;
+		case MaskTarget_PhysMesh::MaxDistance:
+			Values = InMeshData.MaxDistances;
+			break;
+		case MaskTarget_PhysMesh::AnimDriveMultiplier:
+			Values = InMeshData.AnimDriveMultipliers;
+			break;
+		default:
+			break;
+	}
+}
+
+void FClothParameterMask_PhysMesh::SetValue(int32 InVertexIndex, float InValue)
+{
+	if(InVertexIndex < Values.Num())
+	{
+		Values[InVertexIndex] = InValue;
+	}
+}
+
+float FClothParameterMask_PhysMesh::GetValue(int32 InVertexIndex) const
+{
+	return InVertexIndex < Values.Num() ? Values[InVertexIndex] : 0.0f;
+}
+
+const TArray<float>& FClothParameterMask_PhysMesh::GetValueArray() const
+{
+	return Values;
+}
+
+void FClothParameterMask_PhysMesh::CalcRanges(float& MinValue, float& MaxValue)
+{
+	MinValue = MAX_flt;
+	MaxValue = -MinValue;
+
+	const float* ValuesPtr = Values.GetData();
+	for (int32 i=0; i < Values.Num(); ++i)
+	{
+		if (ValuesPtr[i] < MinValue)
+		{
+			MinValue = ValuesPtr[i];
+		}
+
+		if (ValuesPtr[i] > MaxValue)
+		{
+			MaxValue = ValuesPtr[i];
+		}
+	}
+}
+
+void FClothParameterMask_PhysMesh::Apply(FClothPhysicalMeshData& InTargetMesh)
+{
+	if(CurrentTarget == MaskTarget_PhysMesh::None)
+	{
+		// Nothing to do here, just return
+		return;
+	}
+
+	const int32 NumValues = Values.Num();
+	const int32 NumTargetMeshVerts = InTargetMesh.Vertices.Num();
+
+	if(NumTargetMeshVerts == NumValues)
+	{
+		TArray<float>* TargetArray = nullptr;
+		switch(CurrentTarget)
+		{
+			case MaskTarget_PhysMesh::MaxDistance:
+				TargetArray = &InTargetMesh.MaxDistances;
+				break;
+			case MaskTarget_PhysMesh::BackstopDistance:
+				TargetArray = &InTargetMesh.BackstopDistances;
+				break;
+			case MaskTarget_PhysMesh::BackstopRadius:
+				TargetArray = &InTargetMesh.BackstopRadiuses;
+				break;
+			case MaskTarget_PhysMesh::AnimDriveMultiplier:
+				TargetArray = &InTargetMesh.AnimDriveMultipliers;
+				break;
+			default:
+				break;
+		}
+
+		if(!TargetArray)
+		{
+			return;
+		}
+
+		check((*TargetArray).Num() == NumValues);
+
+		for(int32 Index = 0; Index < NumTargetMeshVerts; ++Index)
+		{
+			(*TargetArray)[Index] = Values[Index];
+		}
+	}
+	else
+	{
+		UE_LOG(LogClothingAsset, Warning, TEXT("Aborted applying mask to physical mesh at %p, value mismatch (NumValues: %d, NumVerts: %d)."), &InTargetMesh, NumValues, NumTargetMeshVerts);
+	}
+}
+
+#if WITH_EDITORONLY_DATA
+
+void FClothLODData::GetParameterMasksForTarget(const MaskTarget_PhysMesh& InTarget, TArray<FClothParameterMask_PhysMesh*>& OutMasks)
+{
+	for(FClothParameterMask_PhysMesh& Mask : ParameterMasks)
+	{
+		if(Mask.CurrentTarget == InTarget)
+		{
+			OutMasks.Add(&Mask);
+		}
+	}
+}
+
+#endif
+
+#undef LOCTEXT_NAMESPACE
diff --git a/Engine/Source/Runtime/ClothingSystemRuntimeNv/Private/ClothingSimulationFactoryNv.cpp b/Engine/Source/Runtime/ClothingSystemRuntimeNv/Private/ClothingSimulationFactoryNv.cpp
new file mode 100644
--- /dev/null
+++ b/Engine/Source/Runtime/ClothingSystemRuntimeNv/Private/ClothingSimulationFactoryNv.cpp
@@ -0,0 +1,45 @@
+// Copyright 1998-2019 Epic Games, Inc. All Rights Reserved.
+
+#include "ClothingSimulationFactoryNv.h"
+#include "Assets/ClothingAsset.h"
+
+#if WITH_NVCLOTH
+#include "ClothingSimulationNv.h"
+#endif
+
+#include "ClothingSimulationInteractorNv.h"
+#include "UObject/Package.h"
+
+IClothingSimulation* UClothingSimulationFactoryNv::CreateSimulation()
+{
+#if WITH_NVCLOTH
+	FClothingSimulationBase* Simulation = new FClothingSimulationNv();
+	return Simulation;
+#endif
+	return nullptr;
+}
+
+void UClothingSimulationFactoryNv::DestroySimulation(IClothingSimulation* InSimulation)
+{
+#if WITH_NVCLOTH
+	delete InSimulation;
+#endif
+}
+
+bool UClothingSimulationFactoryNv::SupportsAsset(UClothingAssetBase* InAsset)
+{
+#if WITH_NVCLOTH
+	return true;
+#endif
+	return false;
+}
+
+bool UClothingSimulationFactoryNv::SupportsRuntimeInteraction()
+{
+	return true;
+}
+
+UClothingSimulationInteractor* UClothingSimulationFactoryNv::CreateInteractor()
+{
+	return NewObject<UClothingSimulationInteractorNv>(GetTransientPackage());
+}
diff --git a/Engine/Source/Runtime/ClothingSystemRuntimeNv/Private/ClothingSimulationInteractorNv.cpp b/Engine/Source/Runtime/ClothingSystemRuntimeNv/Private/ClothingSimulationInteractorNv.cpp
new file mode 100644
--- /dev/null
+++ b/Engine/Source/Runtime/ClothingSystemRuntimeNv/Private/ClothingSimulationInteractorNv.cpp
@@ -0,0 +1,108 @@
+// Copyright 1998-2019 Epic Games, Inc. All Rights Reserved.
+
+#include "ClothingSimulationInteractorNv.h"
+
+#include "ClothingSimulationNv.h"
+
+void UClothingSimulationInteractorNv::PhysicsAssetUpdated()
+{
+#if WITH_NVCLOTH
+	Commands.Add(NvInteractorCommand::CreateLambda([](FClothingSimulationNv* InSimulation, FClothingSimulationContextNv* InContext)
+	{
+		InSimulation->RefreshPhysicsAsset();
+	}));
+
+	MarkDirty();
+#endif
+}
+
+void UClothingSimulationInteractorNv::ClothConfigUpdated()
+{
+#if WITH_NVCLOTH
+	Commands.Add(NvInteractorCommand::CreateLambda([](FClothingSimulationNv* InSimulation, FClothingSimulationContextNv* InContext)
+	{
+		InSimulation->RefreshClothConfig();
+	}));
+
+	MarkDirty();
+#endif
+}
+
+void UClothingSimulationInteractorNv::Sync(IClothingSimulation* InSimulation, IClothingSimulationContext* InContext)
+{
+#if WITH_NVCLOTH
+	check(InSimulation);
+	check(InContext);
+
+	FClothingSimulationNv* NvSim = static_cast<FClothingSimulationNv*>(InSimulation);
+	FClothingSimulationContextNv* NvContext = static_cast<FClothingSimulationContextNv*>(InContext);
+
+	for(NvInteractorCommand& Command : Commands)
+	{
+		Command.Execute(NvSim, NvContext);
+	}
+	Commands.Reset();
+#endif
+}
+
+void UClothingSimulationInteractorNv::SetAnimDriveSpringStiffness(float InStiffness)
+{
+#if WITH_NVCLOTH
+	Commands.Add(NvInteractorCommand::CreateLambda([InStiffness](FClothingSimulationNv* InSimulation, FClothingSimulationContextNv* InContext)
+	{
+		InSimulation->ExecutePerActor([InStiffness](FClothingActorNv& Actor)
+		{
+			Actor.CurrentAnimDriveSpringStiffness = InStiffness;
+		});
+	}));
+
+	MarkDirty();
+#endif
+}
+
+void UClothingSimulationInteractorNv::SetAnimDriveDamperStiffness(float InStiffness)
+{
+#if WITH_NVCLOTH
+	Commands.Add(NvInteractorCommand::CreateLambda([InStiffness](FClothingSimulationNv* InSimulation, FClothingSimulationContextNv* InContext)
+	{
+		InSimulation->ExecutePerActor([InStiffness](FClothingActorNv& Actor)
+		{
+			Actor.CurrentAnimDriveDamperStiffness = InStiffness;
+		});
+	}));
+
+	MarkDirty();
+#endif
+}
+void UClothingSimulationInteractorNv::EnableGravityOverride(const FVector& InVector)
+{
+#if WITH_NVCLOTH
+	Commands.Add(NvInteractorCommand::CreateLambda([InVector](FClothingSimulationNv* InSimulation, FClothingSimulationContextNv* InContext)
+	{
+		InSimulation->ExecutePerActor([InVector](FClothingActorNv& Actor)
+		{
+			Actor.bUseGravityOverride = true;
+			Actor.GravityOverride = InVector;
+		});
+	}));
+
+	MarkDirty();
+#endif
+}
+
+void UClothingSimulationInteractorNv::DisableGravityOverride()
+{
+#if WITH_NVCLOTH
+	Commands.Add(NvInteractorCommand::CreateLambda([](FClothingSimulationNv* InSimulation, FClothingSimulationContextNv* InContext)
+	{
+		InSimulation->ExecutePerActor([](FClothingActorNv& Actor)
+		{
+			Actor.bUseGravityOverride = false;
+			Actor.GravityOverride = FVector(0.0f, 0.0f, 0.0f);
+		});
+	}));
+
+	MarkDirty();
+#endif
+}
+
