Copying //UE4/Dev-Physics to Dev-Main (//UE4/Dev-Main) [at] 9098278

diff --git a/Engine/Source/Runtime/ClothingSystemRuntimeNv/Private/ClothingSimulationNv.cpp b/Engine/Source/Runtime/ClothingSystemRuntimeNv/Private/ClothingSimulationNv.cpp
new file mode 100644
--- /dev/null
+++ b/Engine/Source/Runtime/ClothingSystemRuntimeNv/Private/ClothingSimulationNv.cpp
@@ -0,0 +1,1927 @@
+// Copyright 1998-2019 Epic Games, Inc. All Rights Reserved.
+
+#include "ClothingSimulationNv.h"
+
+#if WITH_NVCLOTH
+
+#include "Assets/ClothingAsset.h"
+#include "ClothingSystemRuntimeModuleNv.h"
+#include "NvClothIncludes.h"
+#include "NvClothSupport.h"
+
+#include "Utils/ClothingMeshUtils.h"
+
+#include "PhysicsPublic.h"
+#include "PhysXPublic.h"
+
+#include "DrawDebugHelpers.h"
+#include "Modules/ModuleManager.h"
+
+#include "Components/SkeletalMeshComponent.h"
+
+#if WITH_EDITOR
+#include "SceneManagement.h"
+#include "DynamicMeshBuilder.h"
+#endif
+
+#include "PhysicsEngine/PhysicsAsset.h"
+#include "PhysicsEngine/PhysicsSettings.h"
+#include "SkeletalRenderPublic.h"
+
+DECLARE_CYCLE_STAT(TEXT("Compute Clothing Normals"), STAT_NvClothComputeNormals, STATGROUP_Physics);
+DECLARE_CYCLE_STAT(TEXT("Internal Solve"), STAT_NvClothInternalSolve, STATGROUP_Physics);
+DECLARE_CYCLE_STAT(TEXT("Update Collisions"), STAT_NvClothUpdateCollisions, STATGROUP_Physics);
+DECLARE_CYCLE_STAT(TEXT("Fill Context"), STAT_NvClothFillContext, STATGROUP_Physics);
+DECLARE_CYCLE_STAT(TEXT("Update Anim Drive"), STAT_NvClothUpdateAnimDrive, STATGROUP_Physics);
+
+//=============================================================================
+// FClothingSimulationBase
+//=============================================================================
+
+static TAutoConsoleVariable<float> GClothMaxDeltaTimeTeleportMultiplier(
+	TEXT("p.Cloth.MaxDeltaTimeTeleportMultiplier"),
+	1.5f,
+	TEXT("A multiplier of the MaxPhysicsDelta time at which we will automatically just teleport cloth to its new location\n")
+	TEXT(" default: 1.5"));
+
+FClothingSimulationBase::FClothingSimulationBase()
+{
+	MaxPhysicsDelta = UPhysicsSettings::Get()->MaxPhysicsDeltaTime;
+}
+
+FClothingSimulationBase::~FClothingSimulationBase()
+{}
+
+void 
+FClothingSimulationBase::FillContext(
+	USkeletalMeshComponent* InComponent, 
+	float InDeltaTime, 
+	IClothingSimulationContext* InOutContext)
+{
+	LLM_SCOPE(ELLMTag::SkeletalMesh);
+
+	FClothingSimulationContextBase* BaseContext = static_cast<FClothingSimulationContextBase*>(InOutContext);
+	BaseContext->ComponentToWorld = InComponent->GetComponentTransform();
+	BaseContext->PredictedLod = InComponent->PredictedLODLevel;
+	InComponent->GetWindForCloth_GameThread(BaseContext->WindVelocity, BaseContext->WindAdaption);
+	USkeletalMesh* SkelMesh = InComponent->SkeletalMesh;
+
+	if(USkinnedMeshComponent* MasterComponent = InComponent->MasterPoseComponent.Get())
+	{
+		int32 NumBones = InComponent->MasterBoneMap.Num();
+
+		if(NumBones == 0)
+		{
+			if(InComponent->SkeletalMesh)
+			{
+				// This case indicates an invalid master pose component (e.g. no skeletal mesh)
+				NumBones = InComponent->SkeletalMesh->RefSkeleton.GetNum();
+
+				BaseContext->BoneTransforms.Empty(NumBones);
+				BaseContext->BoneTransforms.AddDefaulted(NumBones);
+			}
+		}
+		else
+		{
+		    BaseContext->BoneTransforms.Reset(NumBones);
+		    BaseContext->BoneTransforms.AddDefaulted(NumBones);
+    
+		    for(int32 BoneIndex = 0; BoneIndex < NumBones; ++BoneIndex)
+		    {
+			    bool bFoundMaster = false;
+			    if(InComponent->MasterBoneMap.IsValidIndex(BoneIndex))
+			    {
+				    const int32 MasterIndex = InComponent->MasterBoneMap[BoneIndex];
+    
+				    if(MasterIndex != INDEX_NONE && MasterIndex < MasterComponent->GetComponentSpaceTransforms().Num())
+				    {
+					    BaseContext->BoneTransforms[BoneIndex] = MasterComponent->GetComponentSpaceTransforms()[MasterIndex];
+					    bFoundMaster = true;
+				    }
+			    }
+    
+			    if(!bFoundMaster)
+			    {
+				    if(SkelMesh)
+				    {
+					    const int32 ParentIndex = SkelMesh->RefSkeleton.GetParentIndex(BoneIndex);
+    
+					    if(ParentIndex != INDEX_NONE)
+					    {
+						    BaseContext->BoneTransforms[BoneIndex] = BaseContext->BoneTransforms[ParentIndex] * SkelMesh->RefSkeleton.GetRefBonePose()[BoneIndex];
+					    }
+					    else
+					    {
+						    BaseContext->BoneTransforms[BoneIndex] = SkelMesh->RefSkeleton.GetRefBonePose()[BoneIndex];
+					    }
+				    }
+			    }
+		    }
+		}
+	}
+	else
+	{
+		BaseContext->BoneTransforms = InComponent->GetComponentSpaceTransforms();
+	}
+
+	UWorld* ComponentWorld = InComponent->GetWorld();
+	check(ComponentWorld);
+
+	BaseContext->DeltaSeconds = FMath::Min(InDeltaTime, MaxPhysicsDelta);
+
+	if(InDeltaTime > (MaxPhysicsDelta * GClothMaxDeltaTimeTeleportMultiplier.GetValueOnGameThread()))
+	{
+		BaseContext->TeleportMode = EClothingTeleportMode::Teleport;
+	}
+	else
+	{
+		BaseContext->TeleportMode = InComponent->ClothTeleportMode;
+	}	
+
+	BaseContext->MaxDistanceScale = InComponent->GetClothMaxDistanceScale();
+
+	float GravityStrength = 981.0f;
+
+	if(UWorld* CurrWorld = InComponent->GetWorld())
+	{
+		GravityStrength = CurrWorld->GetGravityZ();
+	}
+
+	BaseContext->WorldGravity = FVector(0.0f, 0.0f, GravityStrength);
+
+	// Checking the component here to track rare issue leading to invalid contexts
+	if(InComponent->IsPendingKill())
+	{
+		AActor* CompOwner = InComponent->GetOwner();
+		UE_LOG(LogSkeletalMesh, Warning, 
+			TEXT("Attempting to fill a clothing simulation context for a PendingKill skeletal mesh component (Comp: %s, Actor: %s). "
+				"Pending kill skeletal mesh components should be unregistered before marked pending kill."), 
+			*InComponent->GetName(), CompOwner ? *CompOwner->GetName() : TEXT("None"));
+
+		// Make sure we clear this out to skip any attempted simulations
+		BaseContext->BoneTransforms.Reset();
+	}
+
+	if(BaseContext->BoneTransforms.Num() == 0)
+	{
+		AActor* CompOwner = InComponent->GetOwner();
+		USkinnedMeshComponent* Master = InComponent->MasterPoseComponent.Get();
+		UE_LOG(LogSkeletalMesh, Warning, TEXT("Attempting to fill a clothing simulation context for a skeletal mesh component that has zero bones (Comp: %s, Master: %s, Actor: %s)."), *InComponent->GetName(), Master ? *Master->GetName() : TEXT("None"), CompOwner ? *CompOwner->GetName() : TEXT("None"));
+
+		// Make sure we clear this out to skip any attempted simulations
+		BaseContext->BoneTransforms.Reset();
+	}
+}
+
+//=============================================================================
+// FClothingSimulationNv
+//=============================================================================
+
+FClothingSimulationNv::FClothingSimulationNv()
+{}
+
+FClothingSimulationNv::~FClothingSimulationNv()
+{}
+
+void FClothingSimulationNv::CreateActor(USkeletalMeshComponent* InOwnerComponent, UClothingAssetBase* InAsset, int32 InSimDataIndex)
+{
+	check(CachedFactory);
+	check(InOwnerComponent);
+
+	UClothingAsset* Asset = Cast<UClothingAsset>(InAsset);
+	
+	UWorld* World = InOwnerComponent->GetWorld();
+	FPhysScene* PhysicsScene = World ? World->GetPhysicsScene() : nullptr;
+
+	if(!PhysicsScene)
+	{
+		// No scene yet, don't create
+		return;
+	}
+
+	if(!World->bShouldSimulatePhysics)
+	{
+		// No physics for this world
+		return;
+	}
+
+	// Need the current reftolocals so we can skin the ref pose for the sim mesh
+	TArray<FMatrix> RefToLocals;
+	InOwnerComponent->GetCurrentRefToLocalMatrices(RefToLocals, FMath::Min(InOwnerComponent->PredictedLODLevel, Asset->LodData.Num() - 1));
+
+	Actors.AddDefaulted();
+	FClothingActorNv& NewActor = Actors.Last();
+
+	NewActor.SimDataIndex = InSimDataIndex;
+
+	for(int32 LodIndex = 0; LodIndex < Asset->LodData.Num(); ++LodIndex)
+	{
+		FClothLODData& AssetLodData = Asset->LodData[LodIndex];
+		FClothPhysicalMeshData& PhysMesh = AssetLodData.PhysicalMeshData;
+
+		NewActor.LodData.AddDefaulted();
+		FClothingActorNv::FActorLodData& ActorLodData = NewActor.LodData.Last();
+
+		TArray<FVector> Verts;
+		TArray<FVector> SkinnedVerts;
+		TArray<FVector> SkinnedNormals;
+		TArray<NvClothSupport::ClothTri> Tris;
+		TArray<NvClothSupport::ClothQuad> Quads;
+		TArray<float> InvMasses = PhysMesh.InverseMasses;
+
+		const int32 NumVerts = PhysMesh.Vertices.Num();
+		const int32 NumTriangles = PhysMesh.Indices.Num() / 3;
+
+		// Copy data from mesh
+		Verts = PhysMesh.Vertices;
+
+		// We need to skin the vert positions to the current pose, or we'll end
+		// up with clothing placed incorrectly on already posed meshes.
+		FTransform SimBoneTransformCS = InOwnerComponent->GetBoneTransform(Asset->ReferenceBoneIndex, FTransform::Identity);
+		ClothingMeshUtils::SkinPhysicsMesh(Asset->UsedBoneIndices, PhysMesh, SimBoneTransformCS, RefToLocals.GetData(), RefToLocals.Num(), SkinnedVerts, SkinnedNormals);
+
+		Tris.AddDefaulted(NumTriangles);
+
+		for(int32 TriIdx = 0; TriIdx < NumTriangles; ++TriIdx)
+		{
+			NvClothSupport::ClothTri& Triangle = Tris[TriIdx];
+
+			const int32 BaseIdx = TriIdx * 3;
+
+			Triangle.T[0] = PhysMesh.Indices[BaseIdx];
+			Triangle.T[1] = PhysMesh.Indices[BaseIdx + 1];
+			Triangle.T[2] = PhysMesh.Indices[BaseIdx + 2];
+		}
+
+		// Set up a mesh desc for quadification and cooking
+		nv::cloth::ClothMeshDesc MeshDesc;
+		MeshDesc.points.data = SkinnedVerts.GetData();
+		MeshDesc.points.count = SkinnedVerts.Num();
+		MeshDesc.points.stride = SkinnedVerts.GetTypeSize();
+		MeshDesc.triangles.data = Tris.GetData();
+		MeshDesc.triangles.count = Tris.Num();
+		MeshDesc.triangles.stride = Tris.GetTypeSize();
+
+		// Only set up inverse masses here if we aren't completely skinned, otherwise we will fail
+		// constraint creation
+		bool bHasValidMasses = false;
+		for(const float& InvMass : InvMasses)
+		{
+			if(InvMass > 0.0f)
+			{
+				bHasValidMasses = true;
+				break;
+			}
+		}
+
+		if(bHasValidMasses)
+		{
+			MeshDesc.invMasses.data = InvMasses.GetData();
+			MeshDesc.invMasses.count = InvMasses.Num();
+			MeshDesc.invMasses.stride = InvMasses.GetTypeSize();
+		}
+		else
+		{
+			MeshDesc.invMasses.data = nullptr;
+			MeshDesc.invMasses.count = 0;
+			MeshDesc.invMasses.stride = 0;
+		}
+
+		// NvCloth works better with quad meshes, so we need to build one from our triangle data
+		FClothingSystemRuntimeModuleNv& ClothingModule = FModuleManager::Get().LoadModuleChecked<FClothingSystemRuntimeModuleNv>("ClothingSystemRuntimeNv");
+		nv::cloth::ClothMeshQuadifier* Quadifier = ClothingModule.GetMeshQuadifier();
+		Quadifier->quadify(MeshDesc);
+
+		nv::cloth::Vector<int32>::Type NvPhaseInfo;
+		nv::cloth::Fabric* Fabric = NvClothCookFabricFromMesh(CachedFactory, Quadifier->getDescriptor(), physx::PxVec3(0.0f, 0.0f, -981.0f), &NvPhaseInfo, true);
+
+		// Pack the inv mass of each vert to build the starting frame for the cloth
+		ActorLodData.Px_RestPositions.Empty();
+		ActorLodData.Px_RestPositions.AddDefaulted(NumVerts);
+		for(int32 VertIdx = 0; VertIdx < NumVerts; ++VertIdx)
+		{
+			ActorLodData.Px_RestPositions[VertIdx] = physx::PxVec4(U2PVector(SkinnedVerts[VertIdx]), PhysMesh.InverseMasses[VertIdx]);
+		}
+
+		nv::cloth::Cloth* NewCloth = CachedFactory->createCloth(NvClothSupport::CreateRange(ActorLodData.Px_RestPositions), *Fabric);
+
+		// Store off the simulation objects
+		ActorLodData.Cloth = NewCloth;
+		ActorLodData.Fabric = Fabric;
+
+		// Set up our phase (constraint) information
+		const int32 NumPhases = Fabric->getNumPhases();
+		ActorLodData.PhaseConfigs.AddDefaulted(NumPhases);
+		ActorLodData.PhaseTypes.AddDefaulted(NumPhases);
+
+		for(int32 PhaseIdx = 0; PhaseIdx < NumPhases; ++PhaseIdx)
+		{
+			// Get the types and the indices, the rest of the setup will be applied in ApplyClothConfig
+			ActorLodData.PhaseTypes[PhaseIdx] = NvPhaseInfo[PhaseIdx];
+			ActorLodData.PhaseConfigs[PhaseIdx].mPhaseIndex = PhaseIdx;
+		}
+
+		// Set self collision indices
+		NewCloth->setSelfCollisionIndices(NvClothSupport::CreateRange(PhysMesh.SelfCollisionIndices));
+
+		// Set up motion constraints (max distances)
+		nv::cloth::Range<physx::PxVec4> MotionConstraints = NewCloth->getMotionConstraints();
+		const int32 NumMotionConstraints = NewCloth->getNumMotionConstraints();
+		check(NumMotionConstraints == Verts.Num());
+		for(int32 ConstraintIndex = 0; ConstraintIndex < NumMotionConstraints; ++ConstraintIndex)
+		{
+			physx::PxVec4& Constraint = MotionConstraints[ConstraintIndex];
+
+			Constraint = physx::PxVec4(U2PVector(SkinnedVerts[ConstraintIndex]), PhysMesh.MaxDistances[ConstraintIndex]);
+		}
+
+		// Set up the starting transform data for the cloth, then clear our inertia so
+		// we don't get a pop on the first frame. The clothing doesn't really have a
+		// transform this is just used to pass inertial effects to the cloth
+		FTransform RootBoneWorldTransform = InOwnerComponent->GetBoneTransform(Asset->ReferenceBoneIndex);
+		NewCloth->setTranslation(U2PVector(RootBoneWorldTransform.GetTranslation()));
+		NewCloth->setRotation(U2PQuat(RootBoneWorldTransform.GetRotation()));
+		NewCloth->clearInertia();
+
+		// Set the fluid density scale correctly for our units (Nv using metres, UE using centimetres)
+		NewCloth->setFluidDensity(1.0f / NvClothSupport::Constants::UnitConversionScaleCube);
+
+		// Keep track of our asset
+		NewActor.AssetCreatedFrom = Asset;
+
+		// LOD0 is responsible on the first frame, so store off current data for frame-0 for LOD0
+		if(LodIndex == 0)
+		{
+			// Initialise storage arrays for data that will be generated during the sim step
+			NewActor.CurrentNormals.AddDefaulted(NumVerts);
+			NewActor.SkinnedPhysicsMeshPositions[0] = SkinnedVerts;
+			NewActor.SkinnedPhysicsMeshPositions[1] = SkinnedVerts;
+			NewActor.SkinnedPhysicsMeshNormals = SkinnedNormals;
+		}
+	}
+
+	ApplyClothConfig(Asset->ClothConfig, NewActor);
+
+	// Pull collisions from the specified physics asset inside the clothing asset
+	ExtractActorCollisions(Asset, NewActor);
+
+	// Invalid indices so the call to UpdateLod runs all the correct logic as if our LOD just changed.
+	NewActor.CurrentLodIndex = INDEX_NONE;
+
+	// Force update LODs so we're in the correct state now, need to resolve MPC if one is present
+	USkinnedMeshComponent* TransformComponent = InOwnerComponent->MasterPoseComponent.IsValid() ? InOwnerComponent->MasterPoseComponent.Get() : InOwnerComponent;
+	UpdateLod(InOwnerComponent->PredictedLODLevel, InOwnerComponent->GetComponentTransform(), TransformComponent->GetComponentSpaceTransforms(), RefToLocals, true, true);
+
+	// Compute normals for all active actors for first frame
+	for(FClothingActorNv& Actor : Actors)
+	{
+		if(Actor.CurrentLodIndex != INDEX_NONE)
+		{
+			ComputePhysicalMeshNormals(Actor);
+		}
+	}
+}
+
+void FClothingSimulationNv::ExtractActorCollisions(UClothingAsset* Asset, FClothingActorNv &InActor)
+{
+	// Build collision data for this Actor
+	USkeletalMesh* TargetMesh = CastChecked<USkeletalMesh>(Asset->GetOuter());
+	InActor.ExtractedCollisions.Reset();
+
+	if(UPhysicsAsset* PhysAsset = Asset->PhysicsAsset)
+	{
+		bool bAddedBodies = false;
+		for(const USkeletalBodySetup* BodySetup : PhysAsset->SkeletalBodySetups)
+		{
+			int32 MeshBoneIndex = TargetMesh->RefSkeleton.FindBoneIndex(BodySetup->BoneName);
+			int32 MappedBoneIndex = INDEX_NONE;
+
+			if(MeshBoneIndex != INDEX_NONE)
+			{
+				MappedBoneIndex = Asset->UsedBoneNames.AddUnique(BodySetup->BoneName);
+			}
+
+			for(const FKSphereElem& Sphere : BodySetup->AggGeom.SphereElems)
+			{
+				FClothCollisionPrim_Sphere NewSphere;
+				NewSphere.LocalPosition = Sphere.Center;
+				NewSphere.Radius = Sphere.Radius;
+				NewSphere.BoneIndex = MappedBoneIndex;
+
+				InActor.ExtractedCollisions.Spheres.Add(NewSphere);
+				bAddedBodies = true;
+			}
+
+			for(const FKSphylElem& Sphyl : BodySetup->AggGeom.SphylElems)
+			{
+				FClothCollisionPrim_Sphere Sphere0;
+				FClothCollisionPrim_Sphere Sphere1;
+				FVector OrientedDirection = Sphyl.Rotation.RotateVector(FVector(0.0f, 0.0f, 1.0f));
+				FVector HalfDim = OrientedDirection * (Sphyl.Length / 2.0f);
+				Sphere0.LocalPosition = Sphyl.Center - HalfDim;
+				Sphere1.LocalPosition = Sphyl.Center + HalfDim;
+				Sphere0.Radius = Sphyl.Radius;
+				Sphere1.Radius = Sphyl.Radius;
+				Sphere0.BoneIndex = MappedBoneIndex;
+				Sphere1.BoneIndex = MappedBoneIndex;
+
+				InActor.ExtractedCollisions.Spheres.Add(Sphere0);
+				InActor.ExtractedCollisions.Spheres.Add(Sphere1);
+
+				FClothCollisionPrim_SphereConnection Connection;
+				Connection.SphereIndices[0] = InActor.ExtractedCollisions.Spheres.Num() - 2;
+				Connection.SphereIndices[1] = InActor.ExtractedCollisions.Spheres.Num() - 1;
+
+				InActor.ExtractedCollisions.SphereConnections.Add(Connection);
+				bAddedBodies = true;
+			}
+
+			for(const FKTaperedCapsuleElem& TaperedCapsule : BodySetup->AggGeom.TaperedCapsuleElems)
+			{
+				FClothCollisionPrim_Sphere Sphere0;
+				FClothCollisionPrim_Sphere Sphere1;
+				FVector OrientedDirection = TaperedCapsule.Rotation.RotateVector(FVector(0.0f, 0.0f, 1.0f));
+				FVector HalfDim = OrientedDirection * (TaperedCapsule.Length / 2.0f);
+				Sphere0.LocalPosition = TaperedCapsule.Center + HalfDim;
+				Sphere1.LocalPosition = TaperedCapsule.Center - HalfDim;
+				Sphere0.Radius = TaperedCapsule.Radius0;
+				Sphere1.Radius = TaperedCapsule.Radius1;
+				Sphere0.BoneIndex = MappedBoneIndex;
+				Sphere1.BoneIndex = MappedBoneIndex;
+
+				InActor.ExtractedCollisions.Spheres.Add(Sphere0);
+				InActor.ExtractedCollisions.Spheres.Add(Sphere1);
+
+				FClothCollisionPrim_SphereConnection Connection;
+				Connection.SphereIndices[0] = InActor.ExtractedCollisions.Spheres.Num() - 2;
+				Connection.SphereIndices[1] = InActor.ExtractedCollisions.Spheres.Num() - 1;
+
+				InActor.ExtractedCollisions.SphereConnections.Add(Connection);
+				bAddedBodies = true;
+			}
+		}
+
+		// Dirty the actor collisions if we've changed the bodies
+		InActor.bCollisionsDirty |= bAddedBodies;
+
+		// If we've used a bone that isn't in our skinned set we will have added entries
+		// to UsedBoneNames, so rebuild the bone mapping so our collisions work at runtime
+		Asset->RefreshBoneMapping(TargetMesh);
+	}
+}
+
+void FClothingSimulationNv::ApplyClothConfig(FClothConfig &Config, FClothingActorNv &InActor)
+{
+	// These calculations convert our damping and drag values to a value closer to the way APEX used them.
+	// This allows us to not break old assets, and expectations from the old system
+	const float DampStiffnesssFreq = 10.0f;
+	const float DampStiffFreqRatio = DampStiffnesssFreq / Config.StiffnessFrequency;
+	const float ExpDampX = DampStiffFreqRatio * FMath::Log2(1.0f - Config.Damping.X);
+	const float ExpDampY = DampStiffFreqRatio * FMath::Log2(1.0f - Config.Damping.Y);
+	const float ExpDampZ = DampStiffFreqRatio * FMath::Log2(1.0f - Config.Damping.Z);
+	const float ExpDragLinX = DampStiffFreqRatio * FMath::Log2(1.0f - Config.LinearDrag.X);
+	const float ExpDragLinY = DampStiffFreqRatio * FMath::Log2(1.0f - Config.LinearDrag.Y);
+	const float ExpDragLinZ = DampStiffFreqRatio * FMath::Log2(1.0f - Config.LinearDrag.Z);
+	const float ExpDragAngX = DampStiffFreqRatio * FMath::Log2(1.0f - Config.AngularDrag.X);
+	const float ExpDragAngY = DampStiffFreqRatio * FMath::Log2(1.0f - Config.AngularDrag.Y);
+	const float ExpDragAngZ = DampStiffFreqRatio * FMath::Log2(1.0f - Config.AngularDrag.Z);
+
+	const float PrecalcLog2 = FMath::Loge(2.0f);
+	const FVector AdjustedDamping(1.0f - FMath::Exp(ExpDampX * PrecalcLog2),
+								  1.0f - FMath::Exp(ExpDampY * PrecalcLog2),
+								  1.0f - FMath::Exp(ExpDampZ * PrecalcLog2));
+
+	const FVector AdjustedDragLin(1.0f - FMath::Exp(ExpDragLinX * PrecalcLog2),
+								  1.0f - FMath::Exp(ExpDragLinY * PrecalcLog2),
+								  1.0f - FMath::Exp(ExpDragLinZ * PrecalcLog2));
+
+	const FVector AdjustedDragAng(1.0f - FMath::Exp(ExpDragAngX * PrecalcLog2),
+								  1.0f - FMath::Exp(ExpDragAngY * PrecalcLog2),
+								  1.0f - FMath::Exp(ExpDragAngZ * PrecalcLog2));
+
+	for(FClothingActorNv::FActorLodData& LodData : InActor.LodData)
+	{
+		nv::cloth::Cloth* LodCloth = LodData.Cloth;
+
+		// Setup phase configs
+		const int32 NumPhases = LodData.PhaseConfigs.Num();
+		check(LodData.PhaseTypes.Num() == NumPhases);
+		for(int32 PhaseIndex = 0; PhaseIndex < NumPhases; ++PhaseIndex)
+		{
+			nv::cloth::PhaseConfig& PhaseConfig = LodData.PhaseConfigs[PhaseIndex];
+			const nv::cloth::ClothFabricPhaseType::Enum PhaseType = (nv::cloth::ClothFabricPhaseType::Enum)LodData.PhaseTypes[PhaseIndex];
+
+			FClothConstraintSetup* ConfigToUse = nullptr;
+
+			switch(PhaseType)
+			{
+				case nv::cloth::ClothFabricPhaseType::eINVALID:
+					check(false);
+					break;
+				case nv::cloth::ClothFabricPhaseType::eVERTICAL:
+					ConfigToUse = &Config.VerticalConstraintConfig;
+					break;
+				case nv::cloth::ClothFabricPhaseType::eHORIZONTAL:
+					ConfigToUse = &Config.HorizontalConstraintConfig;
+					break;
+				case nv::cloth::ClothFabricPhaseType::eBENDING:
+					ConfigToUse = &Config.BendConstraintConfig;
+					break;
+				case nv::cloth::ClothFabricPhaseType::eSHEARING:
+					ConfigToUse = &Config.ShearConstraintConfig;
+					break;
+			}
+
+			check(ConfigToUse);
+
+			PhaseConfig.mStiffness = ConfigToUse->Stiffness;
+			PhaseConfig.mStiffnessMultiplier = ConfigToUse->StiffnessMultiplier;
+			PhaseConfig.mCompressionLimit = ConfigToUse->CompressionLimit;
+			PhaseConfig.mStretchLimit = ConfigToUse->StretchLimit;
+		}
+		LodCloth->setPhaseConfig(NvClothSupport::CreateRange(LodData.PhaseConfigs));
+
+		// Drag and lift coeffs only take effect in accurate wind mode
+		if(Config.WindMethod == EClothingWindMethod::Accurate)
+		{
+			LodCloth->setDragCoefficient(Config.WindDragCoefficient);
+			LodCloth->setLiftCoefficient(Config.WindLiftCoefficient);
+		}
+		else
+		{
+			LodCloth->setDragCoefficient(0.0f);
+			LodCloth->setLiftCoefficient(0.0f);
+		}
+
+		if(InActor.WindMethod == EClothingWindMethod::Legacy && Config.WindMethod == EClothingWindMethod::Accurate)
+		{
+			// Need to clear out particle accelerations here otherwise legacy wind effects will stay applied
+			LodCloth->clearParticleAccelerations();
+		}
+
+		LodCloth->setSolverFrequency(Config.SolverFrequency);
+		LodCloth->setStiffnessFrequency(Config.StiffnessFrequency);
+		LodCloth->setAcceleationFilterWidth(2.0f * (uint32)Config.SolverFrequency);
+
+		LodCloth->setSelfCollisionDistance(Config.SelfCollisionRadius);
+		LodCloth->setSelfCollisionStiffness(Config.SelfCollisionStiffness);
+
+		// If we have self collision, we need to set rest positions for the mesh
+		if(Config.HasSelfCollision())
+		{
+			LodCloth->setRestPositions(NvClothSupport::CreateRange(LodData.Px_RestPositions));
+		}
+
+		LodCloth->setDamping(U2PVector(AdjustedDamping));
+		LodCloth->setFriction(Config.Friction);
+		LodCloth->setLinearDrag(U2PVector(AdjustedDragLin));
+		LodCloth->setAngularDrag(U2PVector(AdjustedDragAng));
+		LodCloth->setLinearInertia(U2PVector(Config.LinearInertiaScale));
+		LodCloth->setAngularInertia(U2PVector(Config.AngularInertiaScale));
+		LodCloth->setCentrifugalInertia(U2PVector(Config.CentrifugalInertiaScale));
+
+		LodCloth->setTetherConstraintScale(Config.TetherLimit);
+		LodCloth->setTetherConstraintStiffness(Config.TetherStiffness);
+	}
+
+	InActor.CollisionThickness = Config.CollisionThickness;
+	InActor.WindMethod = Config.WindMethod;
+	InActor.CurrentAnimDriveSpringStiffness = Config.AnimDriveSpringStiffness;
+	InActor.CurrentAnimDriveDamperStiffness = Config.AnimDriveDamperStiffness;
+}
+
+IClothingSimulationContext* FClothingSimulationNv::CreateContext()
+{
+	return new FClothingSimulationContextNv();
+}
+
+void FClothingSimulationNv::FillContext(USkeletalMeshComponent* InComponent, float InDeltaTime, IClothingSimulationContext* InOutContext)
+{
+	SCOPE_CYCLE_COUNTER(STAT_NvClothFillContext);
+
+	check(InOutContext);
+
+	FClothingSimulationBase::FillContext(InComponent, InDeltaTime, InOutContext);
+
+	// Assume calling code guarantees the safety of this conversion, we should be given the pointer we allocated in CreateContext
+	FClothingSimulationContextNv* NvContext = (FClothingSimulationContextNv*)InOutContext;
+
+	// Get the current ref to locals to skin fixed vertices
+	NvContext->RefToLocals.Reset();
+	InComponent->GetCurrentRefToLocalMatrices(NvContext->RefToLocals, InComponent->PredictedLODLevel);
+}
+
+void FClothingSimulationNv::Initialize()
+{
+	FClothingSystemRuntimeModuleNv& ClothingModule = FModuleManager::Get().LoadModuleChecked<FClothingSystemRuntimeModuleNv>("ClothingSystemRuntimeNv");
+
+	CachedFactory = ClothingModule.GetSoftwareFactory();
+
+	Solver = CachedFactory->createSolver();
+
+	CurrentMeshLodIndex = INDEX_NONE;
+}
+
+void FClothingSimulationNv::Shutdown()
+{
+	DestroyActors();
+
+	delete Solver;
+	Solver = nullptr;
+}
+
+void FClothingSimulationNv::Simulate(IClothingSimulationContext* InContext)
+{
+	FClothingSimulationContextNv* NvContext = (FClothingSimulationContextNv*)InContext;
+
+	if(NvContext->BoneTransforms.Num() == 0)
+	{
+		// We shouldn't hit this case, the context should have been created by this simulation and should only ever be destroyed by this simulation
+		// Similarly skeletal mesh components should always have at least a root bone transform. However there have been rare cases of empty
+		// transforms lists, we try to catch that here to provide more information
+
+		ensureMsgf(false, TEXT("Invalid context passed to clothing simulation, BoneTransforms has zero entries."));
+
+		// Don't perform simulation
+		return;
+	}
+
+	UpdateLod(NvContext->PredictedLod, NvContext->ComponentToWorld, NvContext->BoneTransforms, NvContext->RefToLocals);
+
+	// Pre-sim work
+	for(FClothingActorNv& Actor : Actors)
+	{
+		if(Actor.CurrentLodIndex == INDEX_NONE)
+		{
+			// No lod to update (the skel component is at a lod level with no clothing for this actor)
+			continue;
+		}
+
+		// Set up scratch arrays
+		FClothingActorScratchData& Scratch = Actor.Scratch;
+		Scratch.Reset();
+
+		// Conditional rebuild. if bCollisionsDirty is set, will rebuild the aggregated collisions
+		Actor.ConditionalRebuildCollisions();
+
+		// Component space root bone transform for this actor
+		FTransform RootBoneTransform = NvContext->BoneTransforms[Actor.AssetCreatedFrom->ReferenceBoneIndex];
+
+		// To build motion constraints (max distances) we need to skin the entire physics mesh
+		// this call also updates our fixed particles to avoid iterating the particle list a second time
+		const FClothPhysicalMeshData& PhysMesh = Actor.AssetCreatedFrom->LodData[Actor.CurrentLodIndex].PhysicalMeshData;
+		ClothingMeshUtils::SkinPhysicsMesh(Actor.AssetCreatedFrom->UsedBoneIndices, PhysMesh, RootBoneTransform, NvContext->RefToLocals.GetData(), NvContext->RefToLocals.Num(), Actor.SkinnedPhysicsMeshPositions[Actor.CurrentSkinnedPositionIndex], Actor.SkinnedPhysicsMeshNormals);
+
+		nv::cloth::Cloth* CurrentCloth = Actor.LodData[Actor.CurrentLodIndex].Cloth;
+
+		bool bTeleport = NvContext->TeleportMode > EClothingTeleportMode::None;
+		bool bReset = NvContext->TeleportMode == EClothingTeleportMode::TeleportAndReset;
+
+		if(bReset)
+		{
+			const TArray<FVector>& CurrSkinnedPositions = Actor.GetCurrentSkinnedPositions();
+			nv::cloth::Range<PxVec4> CurrParticles = CurrentCloth->getCurrentParticles();
+			nv::cloth::Range<PxVec4> PrevParticles = CurrentCloth->getPreviousParticles();
+			const int32 NumParticles = CurrentCloth->getNumParticles();
+			check(NumParticles == CurrSkinnedPositions.Num());
+
+			for(int32 ParticleIndex = 0; ParticleIndex < NumParticles; ++ParticleIndex)
+			{
+				CurrParticles[ParticleIndex] = physx::PxVec4(U2PVector(CurrSkinnedPositions[ParticleIndex]), CurrParticles[ParticleIndex].w);
+				PrevParticles[ParticleIndex] = CurrParticles[ParticleIndex];
+			}
+
+			CurrentCloth->clearParticleAccelerations();
+		}
+
+		// Push the component position into the actor, this will set up the forces in local space to simulate the movement
+		FTransform RootBoneWorldTransform = RootBoneTransform * NvContext->ComponentToWorld;
+
+		if(bTeleport)
+		{
+			// Due to how NvCloth handles teleports we normally lose the velocity on a teleport. Here we instead teleport to the new location minus one step
+			// of velocity integration, which will preserve the velocity and help us to avoid pops.
+			CurrentCloth->teleportToLocation(U2PVector(RootBoneWorldTransform.GetTranslation()), U2PQuat(RootBoneWorldTransform.GetRotation()));
+			CurrentCloth->ignoreVelocityDiscontinuity();
+		}
+
+		CurrentCloth->setTranslation(U2PVector(RootBoneWorldTransform.GetTranslation()));
+		CurrentCloth->setRotation(U2PQuat(RootBoneWorldTransform.GetRotation()));
+
+		FClothConfig& ClothConfig = Actor.AssetCreatedFrom->ClothConfig;
+		if (ClothConfig.bUseGravityOverride)
+		{
+			CurrentCloth->setGravity(U2PVector(ClothConfig.GravityOverride));
+		}
+		else if(Actor.bUseGravityOverride)
+		{
+			CurrentCloth->setGravity(U2PVector(ClothConfig.GravityScale * Actor.GravityOverride));
+		}
+		else
+		{
+			CurrentCloth->setGravity(U2PVector(ClothConfig.GravityScale * NvContext->WorldGravity));
+		}
+
+		Actor.UpdateMotionConstraints(NvContext);
+		
+		{
+			SCOPE_CYCLE_COUNTER(STAT_NvClothUpdateCollisions);
+			// Set collision spheres for this frame
+			FClothCollisionData& CollisionData = Actor.AggregatedCollisions;
+
+			for(FClothCollisionPrim_Sphere& Sphere : CollisionData.Spheres)
+			{
+				FVector SphereLocation = Sphere.LocalPosition;
+
+				if(Sphere.BoneIndex != INDEX_NONE)
+				{
+					const int32 MappedIndex = Actor.AssetCreatedFrom->UsedBoneIndices[Sphere.BoneIndex];
+
+					if(MappedIndex != INDEX_NONE)
+					{
+						const FTransform& BoneTransform = NvContext->BoneTransforms[MappedIndex];
+						SphereLocation = BoneTransform.TransformPosition(Sphere.LocalPosition);
+					}
+				}
+
+				SphereLocation = RootBoneTransform.InverseTransformPosition(SphereLocation);
+
+				Scratch.SphereData.Add(physx::PxVec4(U2PVector(SphereLocation), Sphere.Radius + Actor.CollisionThickness));
+			}
+
+			CurrentCloth->setSpheres(NvClothSupport::CreateRange(Scratch.SphereData), 0, CurrentCloth->getNumSpheres());
+
+			const int32 NumCapsules = CollisionData.SphereConnections.Num();
+
+			for(FClothCollisionPrim_SphereConnection& Capsule : CollisionData.SphereConnections)
+			{
+				Scratch.CapsuleSphereIndices.Add(Capsule.SphereIndices[0]);
+				Scratch.CapsuleSphereIndices.Add(Capsule.SphereIndices[1]);
+			}
+
+			CurrentCloth->setCapsules(NvClothSupport::CreateRange(Scratch.CapsuleSphereIndices), 0, CurrentCloth->getNumCapsules());
+
+			for(FClothCollisionPrim_Convex& Convex : CollisionData.Convexes)
+			{
+				if(Scratch.PlaneData.Num() >= 32)
+				{
+					// Skip, there's too many planes to collide against
+					continue;
+				}
+
+				Scratch.ConvexMasks.AddZeroed();
+				uint32& ConvexMask = Scratch.ConvexMasks.Last();
+
+				for(FPlane& ConvexPlane : Convex.Planes)
+				{
+					Scratch.PlaneData.AddDefaulted();
+					physx::PxVec4& NewPlane = Scratch.PlaneData.Last();
+
+					FPlane TempPlane = ConvexPlane.TransformBy(RootBoneTransform.ToMatrixWithScale().Inverse());
+
+					NewPlane.x = TempPlane.X;
+					NewPlane.y = TempPlane.Y;
+					NewPlane.z = TempPlane.Z;
+					NewPlane.w = -TempPlane.W;
+
+					ConvexMask |= (1 << (Scratch.PlaneData.Num() - 1));
+
+					if(Scratch.PlaneData.Num() >= 32)
+					{
+						break;
+					}
+				}
+			}
+
+			CurrentCloth->setPlanes(NvClothSupport::CreateRange(Scratch.PlaneData), 0, CurrentCloth->getNumPlanes());
+			CurrentCloth->setConvexes(NvClothSupport::CreateRange(Scratch.ConvexMasks), 0, CurrentCloth->getNumConvexes());
+		}
+
+		Actor.UpdateWind(NvContext, NvContext->WindVelocity);
+		Actor.UpdateAnimDrive(NvContext);
+
+		// Cache this frames velocity for teleport resolution
+		Actor.LastVelocity = NvContext->DeltaSeconds > 0.0f ? (RootBoneWorldTransform.GetTranslation() - Actor.LastRootTransform.GetTranslation()) / NvContext->DeltaSeconds : FVector(0.0f, 0.0f, 0.0f);
+		// Cache the last root bone transform
+		Actor.LastRootTransform = RootBoneWorldTransform;
+	}
+
+	// Sim
+	{
+		SCOPE_CYCLE_COUNTER(STAT_NvClothInternalSolve);
+		
+		if(Solver->beginSimulation(NvContext->DeltaSeconds))
+		{
+			// In scenes this would be large, and likely should be distributed over threads,
+			// we tend to only simulated per-actor rather than per-scene so this should always be low
+			const int32 ChunkCount = Solver->getSimulationChunkCount();
+
+			for(int32 ChunkIdx = 0; ChunkIdx < ChunkCount; ++ChunkIdx)
+			{
+				Solver->simulateChunk(ChunkIdx);
+			}
+
+			Solver->endSimulation();
+		}
+	}
+
+	// Post-sim work
+	for(FClothingActorNv& Actor : Actors)
+	{
+		if(Actor.CurrentLodIndex == INDEX_NONE)
+		{
+			continue;
+		}
+
+		// Need to compute mesh normals given new positions
+		ComputePhysicalMeshNormals(Actor);
+
+		// Store off the timestep (needed for velocity calculation)
+		Actor.PreviousTimestep = NvContext->DeltaSeconds;
+
+		// Flip the skin buffer
+		Actor.CurrentSkinnedPositionIndex = (Actor.CurrentSkinnedPositionIndex + 1) % 2;
+	}
+}
+
+void FClothingSimulationNv::ComputePhysicalMeshNormals(FClothingActorNv &Actor)
+{
+	SCOPE_CYCLE_COUNTER(STAT_NvClothComputeNormals);
+
+	FMemory::Memzero(Actor.CurrentNormals.GetData(), Actor.CurrentNormals.Num() * sizeof(FVector));
+
+	const int32 CurrentClothingLod = Actor.CurrentLodIndex;
+
+	const uint32 NumParticles = Actor.LodData[CurrentClothingLod].Cloth->getNumParticles();
+	nv::cloth::MappedRange<physx::PxVec4> Particles = Actor.LodData[CurrentClothingLod].Cloth->getCurrentParticles();
+	const TArray<uint32>& Indices = Actor.AssetCreatedFrom->LodData[CurrentClothingLod].PhysicalMeshData.Indices;
+	const uint32 NumIndices = Indices.Num();
+
+	// Using the face normals, calculate normals. These will not be normalized as we're adding together
+	// normals for every shared face a vert has. We'll normalize later
+	for(uint32 BaseIndex = 0; BaseIndex < NumIndices; BaseIndex += 3)
+	{
+		const FVector A = P2UVector(Particles[Indices[BaseIndex]]);
+		const FVector B = P2UVector(Particles[Indices[BaseIndex + 1]]);
+		const FVector C = P2UVector(Particles[Indices[BaseIndex + 2]]);
+
+		FVector AToB = B - A;
+		FVector AToC = C - A;
+
+		FVector TriNormal = FVector::CrossProduct(AToB, AToC);
+
+		for(int32 TriVertIndex = 0; TriVertIndex < 3; ++TriVertIndex)
+		{
+			Actor.CurrentNormals[Indices[BaseIndex + TriVertIndex]] += TriNormal;
+		}
+	}
+
+	// Normalize the results
+	for(uint32 BaseIndex = 0; BaseIndex < NumParticles; ++BaseIndex)
+	{
+		Actor.CurrentNormals[BaseIndex].Normalize();
+	}
+}
+
+void FClothingSimulationNv::DestroyActors()
+{
+	ensure(Solver);
+
+	for(FClothingActorNv& Actor : Actors)
+	{
+		for(FClothingActorNv::FActorLodData& LodData : Actor.LodData)
+		{
+			Solver->removeCloth(LodData.Cloth);
+			LodData.Fabric->decRefCount();
+
+			delete LodData.Cloth;
+
+			LodData.Cloth = nullptr;
+			LodData.Fabric = nullptr;
+		}
+	}
+
+	// More often a reinit, so don't dealloc here.
+	Actors.Empty(Actors.Num());
+}
+
+void FClothingSimulationNv::DestroyContext(IClothingSimulationContext* InContext)
+{
+	delete InContext;
+}
+
+void FClothingSimulationNv::GetSimulationData(TMap<int32, FClothSimulData>& OutData, USkeletalMeshComponent* InOwnerComponent, USkinnedMeshComponent* InOverrideComponent) const
+{
+	const int32 NumActors = Actors.Num();
+
+	if(NumActors == 0 || !InOwnerComponent)
+	{
+		OutData.Reset();
+		return;
+	}
+
+	if(OutData.Num() != NumActors)
+	{
+		OutData.Reset();
+	}
+
+	bool bFoundSimData = false;
+
+	const FTransform& OwnerTransform = InOwnerComponent->GetComponentTransform();
+	
+	for(int32 ActorIdx = 0; ActorIdx < NumActors; ++ActorIdx)
+	{
+		const FClothingActorNv& Actor = Actors[ActorIdx];
+		UClothingAsset* Asset = Actor.AssetCreatedFrom;
+
+		if(Actor.CurrentLodIndex == INDEX_NONE)
+		{
+			continue;
+		}
+
+		FClothSimulData& ClothData = OutData.FindOrAdd(Actor.SimDataIndex);
+		const int32 CurrentClothingLod = Actor.CurrentLodIndex;
+
+		{
+			NvClothSupport::ClothParticleScopeLock ParticleLock(Actor.LodData[CurrentClothingLod].Cloth);
+
+			const TArray<FTransform>& ReadTransformArray = InOverrideComponent ? InOverrideComponent->GetComponentSpaceTransforms() : InOwnerComponent->GetComponentSpaceTransforms();
+
+			if(!ReadTransformArray.IsValidIndex(Asset->ReferenceBoneIndex))
+			{
+				UE_LOG(LogSkeletalMesh, Warning, TEXT("Failed to write back clothing simulation data for component % as bone transforms are invalid."), *InOwnerComponent->GetName());
+				//ensureMsgf(false, TEXT("Failed to write back clothing simulation data for component % as bone transforms are invalid."), *InOwnerComponent->GetName());
+
+				ClothData.Reset();
+
+				return;
+			}
+
+			FTransform RootBoneTransform = InOverrideComponent ? InOverrideComponent->GetComponentSpaceTransforms()[Asset->ReferenceBoneIndex] : InOwnerComponent->GetComponentSpaceTransforms()[Asset->ReferenceBoneIndex];
+			RootBoneTransform.SetScale3D(FVector(1.0f));
+			RootBoneTransform *= OwnerTransform;
+
+			const uint32 NumParticles = Actor.LodData[CurrentClothingLod].Cloth->getNumParticles();
+			
+			ClothData.Reset();
+			
+			ClothData.Transform = RootBoneTransform;
+			ClothData.ComponentRelativeTransform = RootBoneTransform.GetRelativeTransform(OwnerTransform);
+
+			nv::cloth::MappedRange<physx::PxVec4> Particles = Actor.LodData[CurrentClothingLod].Cloth->getCurrentParticles();
+			for(uint32 ParticleIdx = 0; ParticleIdx < NumParticles; ++ParticleIdx)
+			{
+				physx::PxVec4& Particle = Particles[ParticleIdx];
+
+				ClothData.Positions.Add(P2UVector(Particle));
+				ClothData.Normals.Add(Actor.CurrentNormals[ParticleIdx]);
+			}
+		}
+		
+	}
+}
+
+void FClothingSimulationNv::AddExternalCollisions(const FClothCollisionData& InData)
+{
+	for(FClothingActorNv& Actor : Actors)
+	{
+		Actor.ExternalCollisions.Append(InData);
+		Actor.bCollisionsDirty = true;
+	}
+}
+
+void FClothingSimulationNv::ClearExternalCollisions()
+{
+	for(FClothingActorNv& Actor : Actors)
+	{
+		Actor.ExternalCollisions.Reset();
+		Actor.bCollisionsDirty = true;
+	}
+}
+
+void FClothingSimulationNv::GetCollisions(FClothCollisionData& OutCollisions, bool bIncludeExternal /*= true*/) const
+{
+	OutCollisions.Reset();
+
+	for(const FClothingActorNv& Actor : Actors)
+	{
+		if(bIncludeExternal)
+		{
+			OutCollisions.Append(Actor.AggregatedCollisions);
+		}
+		else
+		{
+			check(Actor.AssetCreatedFrom);
+			OutCollisions.Append(Actor.AssetCreatedFrom->LodData[Actor.CurrentLodIndex].CollisionData);
+		}
+	}
+}
+
+void FClothingSimulationNv::GatherStats() const
+{
+	INC_DWORD_STAT_BY(STAT_NumCloths, Actors.Num());
+
+	for(const FClothingActorNv& Actor : Actors)
+	{
+		// Only count the actor if it's valid.
+		if(!Actor.AssetCreatedFrom || Actor.CurrentLodIndex == INDEX_NONE)
+		{
+			continue;
+		}
+
+		INC_DWORD_STAT_BY(STAT_NumClothVerts, Actor.AssetCreatedFrom->LodData[Actor.CurrentLodIndex].PhysicalMeshData.Vertices.Num());
+	}
+}
+
+bool FClothingSimulationNv::ShouldSimulate() const
+{
+	for(const FClothingActorNv& Actor : Actors)
+	{
+		if(Actor.CurrentLodIndex != INDEX_NONE)
+		{
+			return true;
+		}
+	}
+
+	return false;
+}
+
+FBoxSphereBounds FClothingSimulationNv::GetBounds(const USkeletalMeshComponent* InOwnerComponent) const
+{
+	FBoxSphereBounds CurrentBounds(FVector(0.0f), FVector(0.0f), 0.0f);
+	bool bFirstActor = true;
+
+	const bool bUsingMaster = InOwnerComponent->MasterPoseComponent.IsValid();
+	const USkinnedMeshComponent* ActualComponent = bUsingMaster ? InOwnerComponent->MasterPoseComponent.Get() : InOwnerComponent;
+
+	const TArray<FTransform>& CSTransforms = ActualComponent->GetComponentSpaceTransforms();
+
+	if(CSTransforms.Num() == 0)
+	{
+		const bool bRegistered = ActualComponent->IsRegistered();
+		FString ComponentName = ActualComponent->GetName();
+		USkeletalMesh* CurrentMesh = ActualComponent->SkeletalMesh;
+		FString MeshName = CurrentMesh ? CurrentMesh->GetName() : FString(TEXT("No Mesh"));
+
+		UE_LOG(LogSkeletalMesh, Warning, TEXT("Attempted to calculate clothing bounds for a skeletal mesh that has no component transforms. Registered=%s, Component=%s, Mesh=%s"), bRegistered ? TEXT("True") : TEXT("False"), *ComponentName, *MeshName);
+
+		return CurrentBounds;
+	}
+
+	for(const FClothingActorNv& Actor : Actors)
+	{
+		// Only do this for active actors
+		if(Actor.CurrentLodIndex != INDEX_NONE)
+		{
+			int32 SimBoneIndex = Actor.AssetCreatedFrom->ReferenceBoneIndex;
+
+			if(bUsingMaster)
+			{
+				if(SimBoneIndex < InOwnerComponent->GetMasterBoneMap().Num())
+				{
+					SimBoneIndex = InOwnerComponent->GetMasterBoneMap()[SimBoneIndex];
+					
+					if(SimBoneIndex == INDEX_NONE)
+					{
+						SimBoneIndex = 0;
+					}
+				}
+				else
+				{
+					SimBoneIndex = INDEX_NONE;
+				}
+			}
+
+			FTransform SimBoneTransformCS = SimBoneIndex != INDEX_NONE ? ActualComponent->GetComponentSpaceTransforms()[SimBoneIndex] : FTransform::Identity;
+
+			const FClothingActorNv::FActorLodData& LodData = Actor.LodData[Actor.CurrentLodIndex];
+	
+			FVector Center = SimBoneTransformCS.TransformPosition(P2UVector(LodData.Cloth->getBoundingBoxCenter()));
+			FVector HalfExtents = SimBoneTransformCS.TransformVector(P2UVector(LodData.Cloth->getBoundingBoxScale()));
+	
+			if(bFirstActor)
+			{
+				bFirstActor = false;
+				CurrentBounds = FBox(Center - HalfExtents, Center + HalfExtents);
+			}
+			else
+			{
+				CurrentBounds = CurrentBounds + FBox(Center - HalfExtents, Center + HalfExtents);
+			}
+		}
+	}
+
+	return CurrentBounds;
+}
+
+void FClothingSimulationNv::UpdateLod(int32 InPredictedLod, const FTransform& ComponentToWorld, const TArray<FTransform>& CSTransforms, const TArray<FMatrix>& RefToLocals, bool bForceNoRemap, bool bForceActorChecks)
+{
+	if(InPredictedLod != CurrentMeshLodIndex || bForceActorChecks)
+	{
+		for(FClothingActorNv& Actor : Actors)
+		{
+			const TArray<int32> LodMap = Actor.AssetCreatedFrom->LodMap;
+			if(!LodMap.IsValidIndex(InPredictedLod))
+			{
+				// New LOD unmapped, remove old LOD and move on
+				if(Actor.CurrentLodIndex != INDEX_NONE)
+				{
+					FClothingActorNv::FActorLodData& CurrentLod = Actor.LodData[Actor.CurrentLodIndex];
+					Solver->removeCloth(CurrentLod.Cloth);
+				}
+
+				// Set to none so we don't consider this actor simulated
+				Actor.CurrentLodIndex = INDEX_NONE;
+
+				// Following work is for transitioning between LODs, this actor doesn't require so move on
+				continue;
+			}
+
+			bool bOldLodMapped = LodMap.IsValidIndex(CurrentMeshLodIndex) && LodMap[CurrentMeshLodIndex] != INDEX_NONE;
+
+			// Get the clothing LOD mapped from the mesh predicted LOD
+
+
+			const int32 OldClothingLod = bOldLodMapped ? LodMap[CurrentMeshLodIndex] : INDEX_NONE;
+
+			// If potentialLod doesn't map to a valid LOD, we try higher LOD levels for a valid LOD.
+			// Asset might only have lod on LOD 1 and not 0, however if mesh doesn't force LOD to 1, 
+			// asset will not be assigned valid LOD index and will not generate sim data, breaking things.
+			int32 PredictedClothingLod = INDEX_NONE;
+			for (int32 PotentialLod = InPredictedLod; PotentialLod < LodMap.Num(); ++PotentialLod)
+			{
+				if (LodMap[PotentialLod] != INDEX_NONE)
+				{
+					PredictedClothingLod = LodMap[PotentialLod];
+					break;
+				}
+			}
+
+			if(PredictedClothingLod == Actor.CurrentLodIndex)
+			{
+				// We must have forced a LOD update because we added a new actor - this actor is good to go though.
+				continue;
+			}
+
+			if(!Actor.LodData.IsValidIndex(PredictedClothingLod))
+			{
+				// New LOD not valid for this actor, remove old LOD and move on
+				if(Actor.CurrentLodIndex != INDEX_NONE)
+				{
+					FClothingActorNv::FActorLodData& CurrentLod = Actor.LodData[Actor.CurrentLodIndex];
+					Solver->removeCloth(CurrentLod.Cloth);
+				}
+
+				// Set to none so we don't consider this actor simulated
+				Actor.CurrentLodIndex = INDEX_NONE;
+
+				// Nothing more to do for this actor
+				continue;
+			}
+
+			FClothingActorNv::FActorLodData& NewLodData = Actor.LodData[PredictedClothingLod];
+
+			// Data ranges for the new incoming LOD
+			const int32 NumNewParticles = NewLodData.Cloth->getNumParticles();
+			nv::cloth::Range<physx::PxVec4> NewLodParticles = NewLodData.Cloth->getCurrentParticles();
+			nv::cloth::Range<physx::PxVec4> NewLodPrevParticles = NewLodData.Cloth->getPreviousParticles();
+			nv::cloth::Range<physx::PxVec4> NewAccelerations = NewLodData.Cloth->getParticleAccelerations();
+
+			if(bOldLodMapped && !bForceNoRemap)
+			{
+				FClothingActorNv::FActorLodData& CurrLodData = Actor.LodData[OldClothingLod];
+
+				// the number of LODs we've passed through, we can only reskin the incoming mesh if we've stepped 1 LOD
+				const int32 NumLodsPassed = FMath::Abs(OldClothingLod - PredictedClothingLod);
+
+				const uint32 NumOldParticles = CurrLodData.Cloth->getNumParticles();
+				nv::cloth::Range<const physx::PxVec4> OldLodParticles = nv::cloth::readCurrentParticles(*CurrLodData.Cloth);
+
+				// Remove the old LOD from the solver
+				Solver->removeCloth(Actor.LodData[OldClothingLod].Cloth);
+
+				nv::cloth::Range<physx::PxVec4> OldAccelerations = CurrLodData.Cloth->getParticleAccelerations();
+
+				Solver->addCloth(Actor.LodData[PredictedClothingLod].Cloth);
+
+				if(NumLodsPassed == 1)
+				{
+					// Reposition particles skinned to outgoing LOD
+					bool bLodTransitionUp = OldClothingLod < PredictedClothingLod;
+					FClothLODData& NewLodAssetData = Actor.AssetCreatedFrom->LodData[PredictedClothingLod];
+					TArray<FMeshToMeshVertData>& SkinData = bLodTransitionUp ? NewLodAssetData.TransitionUpSkinData : NewLodAssetData.TransitionDownSkinData;
+
+					for(int32 ParticleIndex = 0; ParticleIndex < NumNewParticles; ++ParticleIndex)
+					{
+						// Do some simple skinning, we only care about positions for this as particles are just
+						// positions inside the solver.
+						FMeshToMeshVertData& VertData = SkinData[ParticleIndex];
+
+						const FVector A = P2UVector(OldLodParticles[VertData.SourceMeshVertIndices[0]]);
+						const FVector B = P2UVector(OldLodParticles[VertData.SourceMeshVertIndices[1]]);
+						const FVector C = P2UVector(OldLodParticles[VertData.SourceMeshVertIndices[2]]);
+
+						// CurrentNormals still contains the normals from the old LOD, which will have been
+						// calculated at the end of the last simulation step.
+						const FVector& NA = Actor.CurrentNormals[VertData.SourceMeshVertIndices[0]];
+						const FVector& NB = Actor.CurrentNormals[VertData.SourceMeshVertIndices[1]];
+						const FVector& NC = Actor.CurrentNormals[VertData.SourceMeshVertIndices[2]];
+
+						const physx::PxVec4& AA = OldAccelerations[VertData.SourceMeshVertIndices[0]];
+						const physx::PxVec4& AB = OldAccelerations[VertData.SourceMeshVertIndices[1]];
+						const physx::PxVec4& AC = OldAccelerations[VertData.SourceMeshVertIndices[2]];
+
+						FVector FinalPosition = VertData.PositionBaryCoordsAndDist.X * A + NA * VertData.PositionBaryCoordsAndDist.W
+							+ VertData.PositionBaryCoordsAndDist.Y * B + NB * VertData.PositionBaryCoordsAndDist.W
+							+ VertData.PositionBaryCoordsAndDist.Z * C + NC * VertData.PositionBaryCoordsAndDist.W;
+
+						physx::PxVec4 FinalAcceleration = VertData.PositionBaryCoordsAndDist.X * AA + VertData.PositionBaryCoordsAndDist.Y * AB + VertData.PositionBaryCoordsAndDist.Z * AC;
+
+						NewLodParticles[ParticleIndex] = physx::PxVec4(U2PVector(FinalPosition), NewLodParticles[ParticleIndex].w);
+						NewLodPrevParticles[ParticleIndex] = physx::PxVec4(U2PVector(FinalPosition), NewLodParticles[ParticleIndex].w);
+						NewAccelerations[ParticleIndex] = FinalAcceleration;
+					}
+				}
+				else
+				{
+					// We've passed more than one LOD, and we don't have transition data for all permutations, just use ref pose
+					for(int32 ParticleIndex = 0; ParticleIndex < NumNewParticles; ++ParticleIndex)
+					{
+						NewLodParticles[ParticleIndex] = NewLodData.Px_RestPositions[ParticleIndex];
+						NewLodPrevParticles[ParticleIndex] = NewLodData.Px_RestPositions[ParticleIndex];
+						NewAccelerations[ParticleIndex] = physx::PxVec4(0.0f);
+					}
+				}
+
+				FTransform SimRootTransform = CSTransforms[Actor.AssetCreatedFrom->ReferenceBoneIndex] * ComponentToWorld;
+				NewLodData.Cloth->setTranslation(U2PVector(SimRootTransform.GetTranslation()));
+				NewLodData.Cloth->setRotation(U2PQuat(SimRootTransform.GetRotation()));
+				NewLodData.Cloth->clearInertia();
+
+				// clear spheres and planes double buffering since cloth kept collision data from last time it was simulated
+				// and clearInertia does not fully empty buffers
+				NewLodData.Cloth->setSpheres(nv::cloth::Range<const PxVec4>(), 0, NewLodData.Cloth->getNumSpheres()); // empty spheres
+				NewLodData.Cloth->setPlanes(nv::cloth::Range<const PxVec4>(), 0, NewLodData.Cloth->getNumPlanes()); // empty planes
+
+				// clear constraints from previous sim
+				NewLodData.Cloth->clearMotionConstraints();
+				NewLodData.Cloth->clearSeparationConstraints();
+
+				Actor.CurrentLodIndex = PredictedClothingLod;
+			}
+			else
+			{
+				// Don't need to do complex remapping as there's no previous clothing LOD, just use ref pose
+				if(Actor.LodData.IsValidIndex(PredictedClothingLod))
+				{
+					Solver->addCloth(NewLodData.Cloth);
+
+					if(CSTransforms.Num() > Actor.AssetCreatedFrom->ReferenceBoneIndex)
+					{
+						// compute skinned positions to init sim mesh
+						const FClothPhysicalMeshData& PhysMesh = Actor.AssetCreatedFrom->LodData[PredictedClothingLod].PhysicalMeshData;
+						TArray<FVector> SkinnedPhysicsMeshPositions;
+						TArray<FVector> SkinnedPhysicsMeshNormals;
+						ClothingMeshUtils::SkinPhysicsMesh(
+							Actor.AssetCreatedFrom->UsedBoneIndices,
+							PhysMesh,
+							CSTransforms[Actor.AssetCreatedFrom->ReferenceBoneIndex],
+							RefToLocals.GetData(),
+							RefToLocals.Num(),
+							SkinnedPhysicsMeshPositions,
+							SkinnedPhysicsMeshNormals);
+
+						for(int32 ParticleIndex = 0; ParticleIndex < NumNewParticles; ++ParticleIndex)
+						{
+							NewLodParticles[ParticleIndex] = PxVec4(U2PVector(SkinnedPhysicsMeshPositions[ParticleIndex]), NewLodParticles[ParticleIndex].w);
+							NewLodPrevParticles[ParticleIndex] = PxVec4(U2PVector(SkinnedPhysicsMeshPositions[ParticleIndex]), NewLodPrevParticles[ParticleIndex].w);
+							NewAccelerations[ParticleIndex] = physx::PxVec4(0.0f);
+						}
+
+						FTransform SimRootTransform = CSTransforms[Actor.AssetCreatedFrom->ReferenceBoneIndex] * ComponentToWorld;
+						NewLodData.Cloth->setTranslation(U2PVector(SimRootTransform.GetTranslation()));
+						NewLodData.Cloth->setRotation(U2PQuat(SimRootTransform.GetRotation()));
+					}
+					else
+					{
+						for(int32 ParticleIndex = 0; ParticleIndex < NumNewParticles; ++ParticleIndex)
+						{
+							NewLodParticles[ParticleIndex] = NewLodData.Px_RestPositions[ParticleIndex];
+							NewLodPrevParticles[ParticleIndex] = NewLodData.Px_RestPositions[ParticleIndex];
+							NewAccelerations[ParticleIndex] = physx::PxVec4(0.0f);
+						}
+					}
+
+					NewLodData.Cloth->clearInertia();
+
+					// clear spheres and planes double buffering since cloth kept collision data from last time it was simulated
+					// and clearInertia does not fully empty buffers
+					NewLodData.Cloth->setSpheres(nv::cloth::Range<const PxVec4>(), 0, NewLodData.Cloth->getNumSpheres()); // empty spheres
+					NewLodData.Cloth->setPlanes(nv::cloth::Range<const PxVec4>(), 0, NewLodData.Cloth->getNumPlanes()); // empty planes
+
+					// clear constraints from previous sim
+					NewLodData.Cloth->clearMotionConstraints();
+					NewLodData.Cloth->clearSeparationConstraints();
+
+					Actor.CurrentLodIndex = PredictedClothingLod;
+				}
+				else
+				{
+					Actor.CurrentLodIndex = INDEX_NONE;
+				}
+			}
+		}
+
+		CurrentMeshLodIndex = InPredictedLod;
+	}
+
+}
+
+#if WITH_EDITOR
+
+void FClothingSimulationNv::DebugDraw_PhysMesh(USkeletalMeshComponent* OwnerComponent, FPrimitiveDrawInterface* PDI) const
+{
+	for(const FClothingActorNv& Actor : Actors)
+	{
+		if(Actor.CurrentLodIndex == INDEX_NONE)
+		{
+			continue;
+		}
+
+		const int32 CurrentClothLod = Actor.CurrentLodIndex;
+
+		nv::cloth::Cloth* CurrentCloth = Actor.LodData[CurrentClothLod].Cloth;
+		
+		check(CurrentCloth);
+
+
+		FTransform RootBoneTransform = OwnerComponent->GetComponentSpaceTransforms()[Actor.AssetCreatedFrom->ReferenceBoneIndex];
+		RootBoneTransform.SetScale3D(FVector(1.0f));
+
+		NvClothSupport::ClothParticleScopeLock ParticleLoc(CurrentCloth);
+
+		const uint32 NumParticles = CurrentCloth->getNumParticles();
+
+		nv::cloth::Range<const physx::PxVec4> Particles = CurrentCloth->getCurrentParticles();
+		if(OwnerComponent->bDisableClothSimulation)
+		{
+			const FClothingActorNv::FActorLodData& ActorData = Actor.LodData[CurrentClothLod];
+			Particles = NvClothSupport::CreateRange(ActorData.Px_RestPositions);
+		}
+
+		TArray<uint32>& Indices = Actor.AssetCreatedFrom->LodData[CurrentClothLod].PhysicalMeshData.Indices;
+		TArray<float>& MaxDistances = Actor.AssetCreatedFrom->LodData[CurrentClothLod].PhysicalMeshData.MaxDistances;
+
+		const int32 NumTriangles = Indices.Num() / 3;
+
+		for(int32 TriangleIndex = 0; TriangleIndex < NumTriangles; ++TriangleIndex)
+		{
+			const int32 BaseIndex = TriangleIndex * 3;
+
+			for(int32 TriVertIndex = 0; TriVertIndex < 3; ++TriVertIndex)
+			{
+				const int32 NextIndex = (TriVertIndex + 1) % 3;
+
+				const FVector Start = RootBoneTransform.TransformPosition(P2UVector(Particles[Indices[BaseIndex + TriVertIndex]]));
+				const FVector End = RootBoneTransform.TransformPosition(P2UVector(Particles[Indices[BaseIndex + NextIndex]]));
+
+				const float MaxDist0 = MaxDistances[Indices[BaseIndex + TriVertIndex]];
+				const float MaxDist1 = MaxDistances[Indices[BaseIndex + NextIndex]];
+
+				const FLinearColor LineColor = MaxDist0 < SMALL_NUMBER && MaxDist1 < SMALL_NUMBER ? FColor::Magenta : FColor::White;
+
+				PDI->DrawLine(Start, End, LineColor, SDPG_World, 0.0f , 0.001f);
+			}
+		}
+	}
+}
+
+void FClothingSimulationNv::DebugDraw_Normals(USkeletalMeshComponent* OwnerComponent, FPrimitiveDrawInterface* PDI) const
+{
+	for(const FClothingActorNv& Actor : Actors)
+	{
+		if(Actor.CurrentLodIndex == INDEX_NONE)
+		{
+			continue;
+		}
+
+		nv::cloth::Cloth* CurrentCloth = Actor.LodData[Actor.CurrentLodIndex].Cloth;
+
+		check(CurrentCloth);
+
+		NvClothSupport::ClothParticleScopeLock ParticleLoc(CurrentCloth);
+
+		const uint32 NumParticles = CurrentCloth->getNumParticles();
+		nv::cloth::MappedRange<physx::PxVec4> Particles = CurrentCloth->getCurrentParticles();
+		const TArray<FVector>& Normals = Actor.CurrentNormals;
+		FTransform RootBoneTransform = OwnerComponent->GetComponentSpaceTransforms()[Actor.AssetCreatedFrom->ReferenceBoneIndex];
+
+		for(uint32 ParticleIndex = 0; ParticleIndex < NumParticles; ++ParticleIndex)
+		{
+			FVector Position = RootBoneTransform.TransformPosition(P2UVector(Particles[ParticleIndex]));
+			FVector Normal = RootBoneTransform.TransformVector(Normals[ParticleIndex]);
+
+			PDI->DrawLine(Position, Position + Normal * 20.0f, FLinearColor::White, SDPG_World, 0.0f, 0.001f);
+		}
+	}
+}
+
+void FClothingSimulationNv::DebugDraw_Collision(USkeletalMeshComponent* OwnerComponent, FPrimitiveDrawInterface* PDI) const
+{
+	for(const FClothingActorNv& Actor : Actors)
+	{
+		if(Actor.CurrentLodIndex == INDEX_NONE)
+		{
+			continue;
+		}
+
+		const FClothCollisionData& CollisionData = Actor.AggregatedCollisions;
+		const TArray<int32>& UsedBones = Actor.AssetCreatedFrom->UsedBoneIndices;
+
+		for(const FClothCollisionPrim_SphereConnection& Connection : CollisionData.SphereConnections)
+		{
+			const FClothCollisionPrim_Sphere& Sphere0 = CollisionData.Spheres[Connection.SphereIndices[0]];
+			const FClothCollisionPrim_Sphere& Sphere1 = CollisionData.Spheres[Connection.SphereIndices[1]];
+
+			if(UsedBones.IsValidIndex(Sphere0.BoneIndex) && UsedBones.IsValidIndex(Sphere1.BoneIndex))
+			{
+				const int32 MappedIndex0 = UsedBones[Sphere0.BoneIndex];
+				const int32 MappedIndex1 = UsedBones[Sphere1.BoneIndex];
+
+				if(MappedIndex0 != INDEX_NONE && MappedIndex1 != INDEX_NONE)
+				{
+					FVector Center0 = OwnerComponent->GetBoneMatrix(MappedIndex0).TransformPosition(Sphere0.LocalPosition);
+					FVector Center1 = OwnerComponent->GetBoneMatrix(MappedIndex1).TransformPosition(Sphere1.LocalPosition);
+
+					// Draws just the sides of a tapered capsule specified by provided Spheres that can have different radii.  Does not draw the spheres, just the sleeve.
+					// Extent geometry endpoints not necessarily coplanar with sphere origins (uses hull horizon)
+					// Otherwise uses the great-circle cap assumption.
+					const float AngleIncrement = 30.0f;   // if parameter added for number of sides, then set this to be:  = 360.0f / NumSides; 
+					FVector Separation = Center1 - Center0;
+					float Distance = Separation.Size();
+					if(Separation.IsNearlyZero() || Distance <= FMath::Abs(Sphere0.Radius - Sphere1.Radius))
+					{
+						continue;
+					}
+					FQuat CapsuleOrientation = FQuat::FindBetween(FVector(0, 0, 1), Separation.GetSafeNormal());
+					float OffsetZ = true ? -(Sphere1.Radius - Sphere0.Radius) / Distance : 0.0f;
+					float ScaleXY = FMath::Sqrt(1.0f - FMath::Square(OffsetZ));
+					FVector VertexPrevious = CapsuleOrientation.RotateVector(FVector(ScaleXY, 0, OffsetZ));
+					for(float Angle = AngleIncrement; Angle <= 360.0f; Angle += AngleIncrement)  // iterate over unit circle about capsule's major axis (which is orientation.AxisZ)
+					{
+						FVector VertexCurrent = CapsuleOrientation.RotateVector(FVector(FMath::Cos(FMath::DegreesToRadians(Angle))*ScaleXY, FMath::Sin(FMath::DegreesToRadians(Angle))*ScaleXY, OffsetZ));
+						PDI->DrawLine(Center0 + VertexCurrent  * Sphere0.Radius, Center1 + VertexCurrent * Sphere1.Radius, FColor::Cyan, SDPG_World, 0.0f, 0.001f);  // capsule side segment between spheres
+						PDI->DrawLine(Center0 + VertexPrevious * Sphere0.Radius, Center0 + VertexCurrent * Sphere0.Radius, FColor::Cyan, SDPG_World, 0.0f, 0.001f);  // cap-circle segment on sphere S0
+						PDI->DrawLine(Center1 + VertexPrevious * Sphere1.Radius, Center1 + VertexCurrent * Sphere1.Radius, FColor::Cyan, SDPG_World, 0.0f, 0.001f);  // cap-circle segment on sphere S1
+						VertexPrevious = VertexCurrent;
+					}
+				}
+			}
+		}
+
+		for(int32 SphereIndex = 0; SphereIndex < CollisionData.Spheres.Num(); ++SphereIndex)
+		{
+			const FClothCollisionPrim_Sphere& Sphere = CollisionData.Spheres[SphereIndex];
+
+			if(UsedBones.IsValidIndex(Sphere.BoneIndex) && UsedBones[Sphere.BoneIndex] != INDEX_NONE)
+			{
+				FMatrix BoneMatrix = OwnerComponent->GetBoneMatrix(UsedBones[Sphere.BoneIndex]);
+				FVector ActualPosition = BoneMatrix.TransformPosition(Sphere.LocalPosition);
+
+				FTransform SphereTransform(BoneMatrix);
+				SphereTransform.SetTranslation(ActualPosition);
+
+				DrawWireSphere(PDI, SphereTransform, FColor::Cyan, Sphere.Radius, 12, SDPG_World, 0.0f, 0.001f);
+			}
+			else
+			{
+				FVector ActualPosition = Sphere.LocalPosition;
+
+				FTransform SphereTransform;
+				SphereTransform.SetTranslation(ActualPosition);
+
+				DrawWireSphere(PDI, SphereTransform, FColor::Red, Sphere.Radius, 12, SDPG_World, 0.0f, 0.001f);
+			}
+		}
+	}
+}
+
+void FClothingSimulationNv::DebugDraw_Backstops(USkeletalMeshComponent* OwnerComponent, FPrimitiveDrawInterface* PDI) const
+{
+	for(const FClothingActorNv& Actor : Actors)
+	{
+		if(Actor.CurrentLodIndex == INDEX_NONE)
+		{
+			continue;
+		}
+
+		const FClothingActorNv::FActorLodData& LodData = Actor.LodData[Actor.CurrentLodIndex];
+		const UClothingAsset* Asset = Actor.AssetCreatedFrom;
+		const FClothPhysicalMeshData& MeshData = Asset->LodData[Actor.CurrentLodIndex].PhysicalMeshData;
+		FTransform RootBoneTransform = OwnerComponent->GetComponentSpaceTransforms()[Actor.AssetCreatedFrom->ReferenceBoneIndex];
+
+		if(!MeshData.HasBackStops())
+		{
+			continue;
+		}
+
+		const TArray<FVector>& SkinnedPositions = Actor.GetCurrentSkinnedPositions();
+		const int32 NumVerts = SkinnedPositions.Num();
+		check(NumVerts == Actor.SkinnedPhysicsMeshNormals.Num());
+		check(NumVerts == MeshData.MaxDistances.Num());
+
+		for(int32 VertIndex = 0; VertIndex < NumVerts; ++VertIndex)
+		{
+			const FVector& Position = RootBoneTransform.TransformPosition(SkinnedPositions[VertIndex]);
+			const FVector& Normal = RootBoneTransform.TransformVector(Actor.SkinnedPhysicsMeshNormals[VertIndex]);
+			
+			float BackstopDistance = MeshData.BackstopDistances[VertIndex];
+			const float BackstopRadius = MeshData.BackstopRadiuses[VertIndex];
+			const float MaxDistance = MeshData.MaxDistances[VertIndex];
+
+			FColor FixedPointColor = FColor::White;
+
+			if(BackstopDistance > MaxDistance)
+			{
+				BackstopDistance = 0.0f;
+
+				// Change the color so disabled backstops are evident
+				FixedPointColor = FColor::Black;
+			}
+
+			if(BackstopDistance > 0.0f)
+			{
+				FVector Start = Position;
+				FVector End = Start + Normal * BackstopDistance;
+
+				PDI->DrawLine(Start, End, FColor::Red, SDPG_World, 0.0f, 0.001f);
+			}
+			else if(BackstopDistance < 0.0f)
+			{
+				FVector Start = Position;
+				FVector End = Start + Normal * BackstopDistance;
+
+				PDI->DrawLine(Start, End, FColor::Blue, SDPG_World, 0.0f, 0.001f);
+			}
+			else
+			{
+				PDI->DrawPoint(Position, FixedPointColor, 2.0f, SDPG_World);
+			}
+		}
+	}
+}
+
+void FClothingSimulationNv::DebugDraw_MaxDistances(USkeletalMeshComponent* OwnerComponent, FPrimitiveDrawInterface* PDI) const
+{
+	for(const FClothingActorNv& Actor : Actors)
+	{
+		if(Actor.CurrentLodIndex == INDEX_NONE)
+		{
+			continue;
+		}
+
+		const FClothingActorNv::FActorLodData& LodData = Actor.LodData[Actor.CurrentLodIndex];
+		const UClothingAsset* Asset = Actor.AssetCreatedFrom;
+		const FClothPhysicalMeshData& MeshData = Asset->LodData[Actor.CurrentLodIndex].PhysicalMeshData;
+		FTransform RootBoneTransform = OwnerComponent->GetComponentSpaceTransforms()[Actor.AssetCreatedFrom->ReferenceBoneIndex];
+
+		const TArray<FVector>& SkinnedPositions = Actor.GetCurrentSkinnedPositions();
+		const int32 NumVerts = SkinnedPositions.Num();
+		check(NumVerts == Actor.SkinnedPhysicsMeshNormals.Num());
+		check(NumVerts == MeshData.MaxDistances.Num());
+
+		for(int32 VertIndex = 0; VertIndex < NumVerts; ++VertIndex)
+		{
+			const FVector& Position = RootBoneTransform.TransformPosition(SkinnedPositions[VertIndex]);
+			const FVector& Normal = RootBoneTransform.TransformVector(Actor.SkinnedPhysicsMeshNormals[VertIndex]);
+			const float& MaxDistance = MeshData.MaxDistances[VertIndex];
+
+			PDI->DrawLine(Position, Position + Normal * MaxDistance, FColor::White, SDPG_World, 0.0f, 0.001f);
+		}
+	}
+}
+
+void FClothingSimulationNv::DebugDraw_SelfCollision(USkeletalMeshComponent* OwnerComponent, FPrimitiveDrawInterface* PDI) const
+{
+	for(const FClothingActorNv& Actor : Actors)
+	{
+		if(Actor.CurrentLodIndex == INDEX_NONE)
+		{
+			// Actor not currently valid due to LOD system
+			continue;
+		}
+
+		const UClothingAsset* Asset = Actor.AssetCreatedFrom;
+		const FClothConfig& Config = Asset->ClothConfig;
+
+		if(!Config.HasSelfCollision())
+		{
+			// No self collisions on this actor
+			continue;
+		}
+
+		FTransform RootBoneTransform = OwnerComponent->GetComponentSpaceTransforms()[Actor.AssetCreatedFrom->ReferenceBoneIndex];
+
+		const float SelfCollisionThickness = Config.SelfCollisionRadius;
+
+		const FClothLODData& LodData = Asset->LodData[Actor.CurrentLodIndex];
+		const FClothPhysicalMeshData& PhysMesh = LodData.PhysicalMeshData;
+
+		nv::cloth::Cloth* CurrentCloth = Actor.LodData[Actor.CurrentLodIndex].Cloth;
+
+		check(CurrentCloth);
+
+		NvClothSupport::ClothParticleScopeLock ParticleLoc(CurrentCloth);
+
+		const uint32 NumParticles = CurrentCloth->getNumParticles();
+		nv::cloth::MappedRange<physx::PxVec4> Particles = CurrentCloth->getCurrentParticles();
+
+		for(int32 SelfColIdx = 0; SelfColIdx < PhysMesh.SelfCollisionIndices.Num(); ++SelfColIdx)
+		{
+			FVector ParticlePosition = RootBoneTransform.TransformPosition(P2UVector(Particles[PhysMesh.SelfCollisionIndices[SelfColIdx]]));
+			DrawWireSphere(PDI, ParticlePosition, FColor::White, SelfCollisionThickness, 8, SDPG_World, 0.0f, 0.001f);
+		}
+	}
+}
+
+void FClothingSimulationNv::DebugDraw_AnimDrive(USkeletalMeshComponent* OwnerComponent, FPrimitiveDrawInterface* PDI) const
+{
+	for(const FClothingActorNv& Actor : Actors)
+	{
+		if(Actor.CurrentLodIndex == INDEX_NONE)
+		{
+			// Actor not currently valid due to LOD system
+			continue;
+		}
+
+		const UClothingAsset* Asset = Actor.AssetCreatedFrom;
+		const FClothPhysicalMeshData& PhysMesh = Asset->LodData[Actor.CurrentLodIndex].PhysicalMeshData;
+		const TArray<uint32>& Indices = PhysMesh.Indices;
+		FTransform RootBoneTransform = OwnerComponent->GetComponentSpaceTransforms()[Actor.AssetCreatedFrom->ReferenceBoneIndex];
+
+		const TArray<FVector> SkinnedPositions = Actor.GetCurrentSkinnedPositions();
+
+		const int32 NumMeshTris = PhysMesh.Indices.Num() / 3;
+
+		for(int32 TriIdx = 0; TriIdx < NumMeshTris; ++TriIdx)
+		{
+			const int32 BaseIndex = TriIdx * 3;
+
+			for(int32 SubIndex = 0; SubIndex < 3; ++SubIndex)
+			{
+				const int32 NextIndex = (SubIndex + 1) % 3;
+
+				FVector Start = RootBoneTransform.TransformPosition(SkinnedPositions[Indices[BaseIndex + SubIndex]]);
+				FVector End = RootBoneTransform.TransformPosition(SkinnedPositions[Indices[BaseIndex + NextIndex]]);
+
+				PDI->DrawLine(Start, End, FLinearColor(0.6f, 0.6f, 0.6f, 0.1f), SDPG_World, 0.0f, 0.001f);
+			}
+		}
+	}
+}
+
+#endif
+
+void FClothingSimulationNv::RefreshClothConfig()
+{
+	checkSlow(IsInGameThread());
+
+	for(FClothingActorNv& Actor : Actors)
+	{
+		ApplyClothConfig(Actor.AssetCreatedFrom->ClothConfig, Actor);
+	}
+}
+
+void FClothingSimulationNv::RefreshPhysicsAsset()
+{
+	checkSlow(IsInGameThread());
+
+	for(FClothingActorNv& Actor : Actors)
+	{
+		ExtractActorCollisions(Actor.AssetCreatedFrom, Actor);
+	}
+}
+
+FClothingActorNv::FClothingActorNv() 
+	: CurrentAnimDriveSpringStiffness(0.0f)
+	, CurrentAnimDriveDamperStiffness(0.0f)
+	, bUseGravityOverride(false)
+	, GravityOverride(0.0f, 0.0f, 0.0f)
+	, LastVelocity(FVector::ZeroVector)
+	, LastRootTransform(FTransform::Identity)
+	, CurrentLodIndex(INDEX_NONE)
+	, bCollisionsDirty(true)
+	, SimDataIndex(INDEX_NONE)
+	, WindMethod(EClothingWindMethod::Legacy)
+	, CurrentSkinnedPositionIndex(0)
+	, PreviousTimestep(0.0f)
+{
+	
+}
+
+void FClothingActorNv::SkinPhysicsMesh(FClothingSimulationContextNv* InContext)
+{
+	const FClothPhysicalMeshData& PhysMesh = AssetCreatedFrom->LodData[CurrentLodIndex].PhysicalMeshData;
+	FTransform RootBoneTransform = InContext->BoneTransforms[AssetCreatedFrom->ReferenceBoneIndex];
+	ClothingMeshUtils::SkinPhysicsMesh(AssetCreatedFrom->UsedBoneIndices, PhysMesh, RootBoneTransform, InContext->RefToLocals.GetData(), InContext->RefToLocals.Num(), SkinnedPhysicsMeshPositions[CurrentSkinnedPositionIndex], SkinnedPhysicsMeshNormals);
+}
+
+void FClothingActorNv::UpdateMotionConstraints(FClothingSimulationContextNv* InContext)
+{
+	if(CurrentLodIndex == INDEX_NONE)
+	{
+		return;
+	}
+
+	nv::cloth::Cloth* CurrentCloth = LodData[CurrentLodIndex].Cloth;
+
+	check(CurrentCloth);
+
+	FClothPhysicalMeshData& PhysMesh = AssetCreatedFrom->LodData[CurrentLodIndex].PhysicalMeshData;
+	TArray<float>& MaxDistances = PhysMesh.MaxDistances;
+	TArray<float>& BackstopDistances = PhysMesh.BackstopDistances;
+	TArray<float>& BackstopRadiuses = PhysMesh.BackstopRadiuses;
+	const TArray<FVector>& CurrentSkinnedPositions = GetCurrentSkinnedPositions();
+
+	nv::cloth::Range<physx::PxVec4> MotionConstraints = CurrentCloth->getMotionConstraints();
+	const int32 NumMotionConstraints = CurrentCloth->getNumMotionConstraints();
+	check(NumMotionConstraints <= CurrentSkinnedPositions.Num());
+
+	for(int32 ConstraintIndex = 0; ConstraintIndex < NumMotionConstraints; ++ConstraintIndex)
+	{
+		MotionConstraints[ConstraintIndex] = physx::PxVec4(U2PVector(CurrentSkinnedPositions[ConstraintIndex]), MaxDistances[ConstraintIndex] * InContext->MaxDistanceScale);
+	}
+
+	if(PhysMesh.HasBackStops())
+	{
+		nv::cloth::Range<physx::PxVec4> SeparationConstraints = CurrentCloth->getSeparationConstraints();
+		const int32 NumSeparationConstraints = CurrentCloth->getNumSeparationConstraints();
+		check(NumSeparationConstraints <= SkinnedPhysicsMeshNormals.Num());
+
+		for(int32 ConstraintIndex = 0; ConstraintIndex < NumSeparationConstraints; ++ConstraintIndex)
+		{
+			SeparationConstraints[ConstraintIndex] = physx::PxVec4(U2PVector(CurrentSkinnedPositions[ConstraintIndex] - BackstopDistances[ConstraintIndex] * SkinnedPhysicsMeshNormals[ConstraintIndex]), BackstopRadiuses[ConstraintIndex]);
+		}
+	}
+}
+
+void FClothingActorNv::UpdateWind(FClothingSimulationContextNv* InContext, const FVector& InWindVelocity)
+{
+	switch(WindMethod)
+	{
+		default:
+			break;
+
+		case EClothingWindMethod::Legacy:
+		{
+			const FTransform RootBoneWorldTransform = InContext->BoneTransforms[AssetCreatedFrom->ReferenceBoneIndex] * InContext->ComponentToWorld;
+			const FVector TransformedWindVelocity = RootBoneWorldTransform.InverseTransformVector(InWindVelocity);
+
+			TArray<FVector>& ParticleVelocities = Scratch.ParticleVelocities;
+			CalculateParticleVelocities(ParticleVelocities);
+
+			const TArray<float>& MaxDistances = AssetCreatedFrom->LodData[CurrentLodIndex].PhysicalMeshData.MaxDistances;
+
+			const int32 NumAccelerations = LodData[CurrentLodIndex].Cloth->getNumParticleAccelerations();
+			nv::cloth::Range<physx::PxVec4> ParticleAccelerations = LodData[CurrentLodIndex].Cloth->getParticleAccelerations();
+
+			for(int32 AccelerationIndex = 0; AccelerationIndex < NumAccelerations; ++AccelerationIndex)
+			{
+				const FVector& Velocity = ParticleVelocities[AccelerationIndex];
+				FVector VelocityDelta = (TransformedWindVelocity * 2500.0f - Velocity);
+
+				if(MaxDistances[AccelerationIndex] > 0.0f && !VelocityDelta.IsZero())
+				{
+					// scaled by angle
+					const float DirectionDot = FVector::DotProduct(VelocityDelta.GetUnsafeNormal(), CurrentNormals[AccelerationIndex]);
+					VelocityDelta *= FMath::Min(1.0f, FMath::Abs(DirectionDot) * InContext->WindAdaption);
+					ParticleAccelerations[AccelerationIndex] = physx::PxVec4(VelocityDelta.X, VelocityDelta.Y, VelocityDelta.Z, 0.0f);
+				}
+				else
+				{
+					ParticleAccelerations[AccelerationIndex].setZero();
+				}
+			}
+		}
+		break;
+
+		case EClothingWindMethod::Accurate:
+		{
+			const physx::PxVec3 PxWindVelocity = U2PVector(InWindVelocity * NvClothSupport::Constants::UnitConversionScale);
+			LodData[CurrentLodIndex].Cloth->setWindVelocity((PxWindVelocity));
+		}
+		break;
+	}
+}
+
+void FClothingActorNv::ConditionalRebuildCollisions()
+{
+	// Only need to rebuild collisions if they're dirty
+	if(!bCollisionsDirty)
+	{
+		return;
+	}
+
+	if(CurrentLodIndex == INDEX_NONE)
+	{
+		return;
+	}
+
+	AggregatedCollisions.Reset();
+
+	// Asset-embedded collisions (created during import)
+	AggregatedCollisions.Append(AssetCreatedFrom->LodData[CurrentLodIndex].CollisionData);
+	// Extracted collisions from the physics asset selected by the user
+	AggregatedCollisions.Append(ExtractedCollisions);
+	// External collisions added from the world
+	AggregatedCollisions.Append(ExternalCollisions);
+
+	bCollisionsDirty = false;
+}
+
+void FClothingActorNv::CalculateParticleVelocities(TArray<FVector>& OutVelocities)
+{
+	float InverseTimestep = PreviousTimestep != 0.0f ? 1.0f / PreviousTimestep : 0.0f;
+
+	const int32 NumParticles = LodData[CurrentLodIndex].Cloth->getNumParticles();
+	nv::cloth::Range<physx::PxVec4> PreviousPositions = LodData[CurrentLodIndex].Cloth->getPreviousParticles();
+	nv::cloth::Range<physx::PxVec4> CurrentPositions = LodData[CurrentLodIndex].Cloth->getCurrentParticles();
+
+	if(OutVelocities.Num() != NumParticles)
+	{
+		OutVelocities.Reset();
+		OutVelocities.AddDefaulted(NumParticles);
+	}
+
+	for(int32 ParticleIndex = 0; ParticleIndex < NumParticles; ++ParticleIndex)
+	{
+		FVector OldPosition = P2UVector(PreviousPositions[ParticleIndex]);
+		FVector NewPosition = P2UVector(CurrentPositions[ParticleIndex]);
+
+		OutVelocities[ParticleIndex] = (NewPosition - OldPosition) * InverseTimestep;
+	}
+}
+
+void FClothingActorNv::UpdateAnimDrive(FClothingSimulationContextNv* InContext)
+{
+	SCOPE_CYCLE_COUNTER(STAT_NvClothUpdateAnimDrive);
+
+	FClothPhysicalMeshData& PhysMesh = AssetCreatedFrom->LodData[CurrentLodIndex].PhysicalMeshData;
+
+	if(PhysMesh.HasAnimDrive())
+	{
+		TArray<float>& AnimDrives = PhysMesh.AnimDriveMultipliers;
+		const TArray<FVector> CurrentSkinPositions = GetCurrentSkinnedPositions();
+		const TArray<FVector> PreviousSkinPositions = GetPreviousSkinnedPositions();
+
+		FActorLodData& CurrLodData = LodData[CurrentLodIndex];
+
+		int32 NumAccelerations = CurrLodData.Cloth->getNumParticleAccelerations();
+		checkSlow(NumAccelerations == 0 || NumAccelerations == AnimDrives.Num());
+		nv::cloth::Range<physx::PxVec4> Accelerations = CurrLodData.Cloth->getParticleAccelerations();
+
+		nv::cloth::Range<const physx::PxVec4> Locations = CurrLodData.Cloth->getCurrentParticles();
+		nv::cloth::Range<const physx::PxVec4> PrevLocations = CurrLodData.Cloth->getPreviousParticles();
+
+		for(int32 ParticleIdx = 0; ParticleIdx < NumAccelerations; ++ParticleIdx)
+		{
+			physx::PxVec4 Location = Locations[ParticleIdx];
+
+			// W Component of location is the inverse mass of the particle in NvCloth (0 -> Infinite mass)
+			if(Location.w == 0.0f)
+			{
+				continue;
+			}
+
+			const FVector ToSkinned = CurrentSkinPositions[ParticleIdx] - P2UVector(Location);
+			const FVector ParticleVelocity = P2UVector(Locations[ParticleIdx] - PrevLocations[ParticleIdx]);
+			const FVector SkinnedVelocity = CurrentSkinPositions[ParticleIdx] - PreviousSkinPositions[ParticleIdx];
+
+			// accel = (kX - bv)/m (spring with damper)
+			const FVector AddedAcceleration = ((CurrentAnimDriveSpringStiffness * AnimDrives[ParticleIdx] * ToSkinned) - (CurrentAnimDriveDamperStiffness * (ParticleVelocity - SkinnedVelocity))) / (1.0f / Location.w);
+
+			Accelerations[ParticleIdx] += physx::PxVec4(U2PVector(AddedAcceleration), 0.0f);
+		}
+	}
+}
+
+const TArray<FVector>& FClothingActorNv::GetCurrentSkinnedPositions() const
+{
+	return SkinnedPhysicsMeshPositions[CurrentSkinnedPositionIndex];
+}
+
+const TArray<FVector>& FClothingActorNv::GetPreviousSkinnedPositions() const
+{
+	return SkinnedPhysicsMeshPositions[(CurrentSkinnedPositionIndex + 1) % 2];
+}
+
+void FClothingActorScratchData::Reset()
+{
+	SphereData.Reset();
+	CapsuleSphereIndices.Reset();
+	PlaneData.Reset();
+	ConvexMasks.Reset();
+	ParticleVelocities.Reset();
+}
+
+#endif
diff --git a/Engine/Source/Runtime/ClothingSystemRuntimeNv/Private/ClothingSystemRuntimeModuleNv.cpp b/Engine/Source/Runtime/ClothingSystemRuntimeNv/Private/ClothingSystemRuntimeModuleNv.cpp
new file mode 100644
--- /dev/null
+++ b/Engine/Source/Runtime/ClothingSystemRuntimeNv/Private/ClothingSystemRuntimeModuleNv.cpp
@@ -0,0 +1,170 @@
+// Copyright 1998-2019 Epic Games, Inc. All Rights Reserved.
+
+#include "ClothingSystemRuntimeModuleNv.h"
+#include "ClothingSimulationFactoryNv.h"
+#include "Features/IModularFeatures.h"
+#include "Misc/Paths.h"
+#include "HAL/PlatformProcess.h"
+#include "Modules/ModuleManager.h"
+
+#include "NvClothIncludes.h"
+#include "NvClothSupport.h"
+
+
+IMPLEMENT_MODULE(FClothingSystemRuntimeModuleNv, ClothingSystemRuntimeNv);
+
+FClothingSystemRuntimeModuleNv::FClothingSystemRuntimeModuleNv()
+#if WITH_NVCLOTH
+	: ClothFactory(nullptr)
+	, Quadifier(nullptr)
+#if PLATFORM_WINDOWS
+	, NvClothHandle(0)
+#endif
+#endif
+{
+
+}
+
+void FClothingSystemRuntimeModuleNv::StartupModule()
+{
+#if WITH_NVCLOTH
+	DelayLoadNvCloth();
+
+	NvClothSupport::InitializeNvClothingSystem();
+
+	ClothFactory = NvClothCreateFactoryCPU();
+	Quadifier = NvClothCreateMeshQuadifier();
+#endif
+
+	IModularFeatures::Get().RegisterModularFeature(IClothingSimulationFactoryClassProvider::FeatureName, this);
+}
+
+void FClothingSystemRuntimeModuleNv::ShutdownModule()
+{
+	IModularFeatures::Get().UnregisterModularFeature(IClothingSimulationFactoryClassProvider::FeatureName, this);
+
+#if WITH_NVCLOTH
+	delete Quadifier;
+	Quadifier = nullptr;
+
+	NvClothDestroyFactory(ClothFactory);
+
+	ShutdownNvClothLibs();
+#endif
+}
+
+nv::cloth::Factory* FClothingSystemRuntimeModuleNv::GetSoftwareFactory()
+{
+#if WITH_NVCLOTH
+	ensureMsgf(ClothFactory, TEXT("Cloth SW factory has not been created yet!"));
+	return ClothFactory;
+#endif
+
+	return nullptr;
+}
+
+nv::cloth::ClothMeshQuadifier* FClothingSystemRuntimeModuleNv::GetMeshQuadifier()
+{
+#if WITH_NVCLOTH
+	ensureMsgf(Quadifier, TEXT("Cloth quadifier has not been created yet!"));
+	return Quadifier;
+#endif
+
+	return nullptr;
+}
+
+UClass* FClothingSystemRuntimeModuleNv::GetDefaultSimulationFactoryClass()
+{
+#if WITH_NVCLOTH
+	return UClothingSimulationFactoryNv::StaticClass();
+#endif
+
+	return nullptr;
+}
+
+#if WITH_NVCLOTH
+void FClothingSystemRuntimeModuleNv::DelayLoadNvCloth()
+{
+#if PLATFORM_WINDOWS
+	DelayLoadNvCloth_Windows();
+#elif PLATFORM_MAC
+	DelayLoadNvCloth_Mac();
+#endif
+}
+
+void FClothingSystemRuntimeModuleNv::ShutdownNvClothLibs()
+{
+#if PLATFORM_WINDOWS
+	ShutdownNvCloth_Windows();
+#elif PLATFORM_MAC
+	ShutdownNvCloth_Mac();
+#endif
+}
+
+#if PLATFORM_WINDOWS
+void FClothingSystemRuntimeModuleNv::DelayLoadNvCloth_Windows()
+{
+	FString PhysXBinariesRoot = FPaths::EngineDir() / TEXT("Binaries/ThirdParty/PhysX3/");
+
+#if _MSC_VER >= 1900
+	FString VSDirectory(TEXT("VS2015/"));
+#else
+#error "Unrecognized Visual Studio version."
+#endif
+
+#if PLATFORM_64BITS
+	FString RootPhysXPath(PhysXBinariesRoot + TEXT("Win64/") + VSDirectory);
+	FString ArchName(TEXT("_x64"));
+#else
+	FString RootPhysXPath(PhysXBinariesRoot + TEXT("Win32/") + VSDirectory);
+	FString ArchName(TEXT("_x86"));
+#endif
+
+#ifdef UE_NVCLOTH_SUFFIX
+	FString LibSuffix(TEXT(PREPROCESSOR_TO_STRING(UE_NVCLOTH_SUFFIX)) + ArchName + FString(TEXT(".dll")));
+#else
+	FString LibSuffix(ArchName + TEXT(".dll"));
+#endif
+
+	FString ModulePath = RootPhysXPath + TEXT("NvCloth") + LibSuffix;
+	NvClothHandle = FPlatformProcess::GetDllHandle(*ModulePath);
+	checkf(NvClothHandle, TEXT("Failed to load module: %s"), *ModulePath);
+}
+
+void FClothingSystemRuntimeModuleNv::ShutdownNvCloth_Windows()
+{
+	if(NvClothHandle)
+	{
+		FPlatformProcess::FreeDllHandle(NvClothHandle);
+	}
+}
+
+#endif // PLATFORM_WINDOWS
+
+#if PLATFORM_MAC
+void FClothingSystemRuntimeModuleNv::DelayLoadNvCloth_Mac()
+{
+	FString PhysXBinariesRoot = FPaths::EngineDir() / TEXT("Binaries/ThirdParty/PhysX3/Mac/");
+
+#ifdef UE_NVCLOTH_SUFFIX
+	FString LibSuffix(TEXT(PREPROCESSOR_TO_STRING(UE_NVCLOTH_SUFFIX)) + FString(TEXT(".dylib")));
+#else
+	FString LibSuffix(TEXT(".dylib"));
+#endif
+
+	FString ModulePath = PhysXBinariesRoot + TEXT("libNvCloth") + LibSuffix;
+	NvClothHandle = FPlatformProcess::GetDllHandle(*ModulePath);
+	checkf(NvClothHandle, TEXT("Failed to load module: %s"), *ModulePath);
+}
+
+void FClothingSystemRuntimeModuleNv::ShutdownNvCloth_Mac()
+{
+	if(NvClothHandle)
+	{
+		FPlatformProcess::FreeDllHandle(NvClothHandle);
+	}
+}
+#endif
+
+
+#endif // WITH_NVCLOTH
diff --git a/Engine/Source/Runtime/ClothingSystemRuntimeNv/Private/ClothingSystemRuntimePrivatePCH.h b/Engine/Source/Runtime/ClothingSystemRuntimeNv/Private/ClothingSystemRuntimePrivatePCH.h
new file mode 100644
--- /dev/null
+++ b/Engine/Source/Runtime/ClothingSystemRuntimeNv/Private/ClothingSystemRuntimePrivatePCH.h
@@ -0,0 +1,3 @@
+// Copyright 1998-2019 Epic Games, Inc. All Rights Reserved.
+
+#include "Engine.h"
\ No newline at end of file
diff --git a/Engine/Source/Runtime/ClothingSystemRuntimeNv/Private/NvClothSupport.cpp b/Engine/Source/Runtime/ClothingSystemRuntimeNv/Private/NvClothSupport.cpp
new file mode 100644
--- /dev/null
+++ b/Engine/Source/Runtime/ClothingSystemRuntimeNv/Private/NvClothSupport.cpp
@@ -0,0 +1,156 @@
+// Copyright 1998-2019 Epic Games, Inc. All Rights Reserved.
+
+#include "NvClothSupport.h"
+
+#if WITH_NVCLOTH
+
+#include "NvClothIncludes.h"
+#include "PhysicsPublic.h"
+
+namespace NvClothSupport
+{
+	class FNvClothErrorCallback* GNvClothErrorCallback;
+	class FNvClothAssertHandler* GNvClothAssertHandler;
+
+	ELogVerbosity::Type PxErrorToLogVerbosity(physx::PxErrorCode::Enum InEnum)
+	{
+		switch(InEnum)
+		{
+			case physx::PxErrorCode::eDEBUG_INFO:
+				return ELogVerbosity::Display;
+				break;
+			case physx::PxErrorCode::eDEBUG_WARNING:
+				return ELogVerbosity::Warning;
+				break;
+			case physx::PxErrorCode::eINVALID_PARAMETER:
+				return ELogVerbosity::Warning;
+				break;
+			case physx::PxErrorCode::eINVALID_OPERATION:
+				return ELogVerbosity::Error;
+				break;
+			case physx::PxErrorCode::eOUT_OF_MEMORY:
+				return ELogVerbosity::Fatal;
+				break;
+			case physx::PxErrorCode::eINTERNAL_ERROR:
+				return ELogVerbosity::Error;
+				break;
+			case physx::PxErrorCode::eABORT:
+				return ELogVerbosity::Fatal;
+				break;
+			case physx::PxErrorCode::ePERF_WARNING:
+				return ELogVerbosity::Warning;
+				break;
+			default:
+				break;
+		}
+
+		return ELogVerbosity::Log;
+	}
+
+	FString PxErrorToString(physx::PxErrorCode::Enum InEnum)
+	{
+		switch(InEnum)
+		{
+			case physx::PxErrorCode::eDEBUG_INFO:
+				return TEXT("Info");
+				break;
+			case physx::PxErrorCode::eDEBUG_WARNING:
+				return TEXT("Warning");
+				break;
+			case physx::PxErrorCode::eINVALID_PARAMETER:
+				return TEXT("Invalid Parameter");
+				break;
+			case physx::PxErrorCode::eINVALID_OPERATION:
+				return TEXT("Invalid Operation");
+				break;
+			case physx::PxErrorCode::eOUT_OF_MEMORY:
+				return TEXT("Out of Memory");
+				break;
+			case physx::PxErrorCode::eINTERNAL_ERROR:
+				return TEXT("Internal Error");
+				break;
+			case physx::PxErrorCode::eABORT:
+				return TEXT("Abort");
+				break;
+			case physx::PxErrorCode::ePERF_WARNING:
+				return TEXT("Performance Warning");
+				break;
+			default:
+				break;
+		}
+
+		return TEXT("Unknown");
+	}
+
+	class FNvClothErrorCallback : public physx::PxErrorCallback
+	{
+	public:
+
+		FNvClothErrorCallback()
+		{}
+
+		virtual void reportError(PxErrorCode::Enum code, const char* message, const char* file, int line) override
+		{
+			if(code != PxErrorCode::eNO_ERROR)
+			{
+				// Actual error occurred, push to log
+				ELogVerbosity::Type Verbosity = PxErrorToLogVerbosity(code);
+				GLog->Logf(Verbosity, TEXT("NvCloth: %s, [%s:%d]"), message, file, line);
+			}
+		}
+
+	};
+
+	class FNvClothAssertHandler : public physx::PxAssertHandler
+	{
+
+	public:
+
+		FNvClothAssertHandler()
+		{}
+
+		virtual void operator()(const char* exp, const char* file, int line, bool& ignore)
+		{
+			checkf(false, TEXT("NvCloth Assert: &s [%s:%d]"), exp, file, line);
+		}
+	};
+
+	physx::PxAllocatorCallback* GetAllocator()
+	{
+		return (physx::PxAllocatorCallback*)GPhysXAllocator;
+	}
+
+	void InitializeNvClothingSystem()
+	{
+		GNvClothErrorCallback = new FNvClothErrorCallback();
+		GNvClothAssertHandler = new FNvClothAssertHandler();
+
+		nv::cloth::InitializeNvCloth(GetAllocator(), GNvClothErrorCallback, GNvClothAssertHandler, nullptr);
+	}
+
+	void ShutdownNvClothingSystem()
+	{
+		delete GNvClothAssertHandler;
+		delete GNvClothErrorCallback;
+	}
+
+	ClothParticleScopeLock::ClothParticleScopeLock() : LockedCloth(nullptr)
+	{
+
+	}
+
+	ClothParticleScopeLock::ClothParticleScopeLock(nv::cloth::Cloth* InCloth)
+		: LockedCloth(InCloth)
+	{
+		check(LockedCloth);
+		LockedCloth->lockParticles();
+	}
+
+	ClothParticleScopeLock::~ClothParticleScopeLock()
+	{
+		LockedCloth->unlockParticles();
+	}
+
+} // namespace NvClothSupport
+
+#endif // WITH_NVCLOTH
diff --git a/Engine/Source/Runtime/ClothingSystemRuntimeNv/Private/Utils/ClothingMeshUtils.cpp b/Engine/Source/Runtime/ClothingSystemRuntimeNv/Private/Utils/ClothingMeshUtils.cpp
new file mode 100644
--- /dev/null
+++ b/Engine/Source/Runtime/ClothingSystemRuntimeNv/Private/Utils/ClothingMeshUtils.cpp
@@ -0,0 +1,346 @@
+// Copyright 1998-2019 Epic Games, Inc. All Rights Reserved.
+
+#include "Utils/ClothingMeshUtils.h"
+#include "Assets/ClothingAsset.h" // for FClothPhysicalMeshData
+
+#include "Math/UnrealMathUtility.h"
+#include "Logging/LogMacros.h"
+
+#if WITH_EDITOR
+#include "Framework/Notifications/NotificationManager.h"
+#include "Widgets/Notifications/SNotificationList.h"
+#endif
+
+DEFINE_LOG_CATEGORY(LogClothingMeshUtils)
+DECLARE_CYCLE_STAT(TEXT("Skin Physics Mesh"), STAT_ClothSkinPhysMesh, STATGROUP_Physics);
+
+#define LOCTEXT_NAMESPACE "ClothingMeshUtils"
+
+namespace ClothingMeshUtils
+{
+	void SkinPhysicsMesh(
+		const TArray<int32>& BoneMap,
+		const FClothPhysicalMeshData& InMesh,
+		const FTransform& RootBoneTransform,
+		const FMatrix* InBoneMatrices,
+		const int32 InNumBoneMatrices,
+		TArray<FVector>& OutPositions,
+		TArray<FVector>& OutNormals)
+	{
+		SCOPE_CYCLE_COUNTER(STAT_ClothSkinPhysMesh);
+
+		// Ignore any user scale. It's already accounted for in our skinning matrices
+		FTransform RootBoneTransformInternal = RootBoneTransform;
+		RootBoneTransformInternal.SetScale3D(FVector(1.0f));
+
+		const uint32 NumVerts = InMesh.Vertices.Num();
+
+		OutPositions.Reset(NumVerts);
+		OutNormals.Reset(NumVerts);
+		OutPositions.AddZeroed(NumVerts);
+		OutNormals.AddZeroed(NumVerts);
+
+		const int32 MaxInfluences = InMesh.MaxBoneWeights;
+		const FMatrix* RESTRICT BoneMatrices = InBoneMatrices;
+
+		for (uint32 VertIndex = 0; VertIndex < NumVerts; ++VertIndex)
+		{
+			// Fixed particle, needs to be skinned
+			const uint16* RESTRICT BoneIndices = InMesh.BoneData[VertIndex].BoneIndices;
+			const float* RESTRICT BoneWeights = InMesh.BoneData[VertIndex].BoneWeights;
+
+			// WARNING - HORRIBLE UNROLLED LOOP + JUMP TABLE BELOW
+			// done this way because this is a pretty tight and perf critical loop. essentially
+			// rather than checking each influence we can just jump into this switch and fall through
+			// everything to compose the final skinned data
+			const FVector& RefParticle = InMesh.Vertices[VertIndex];
+			const FVector& RefNormal = InMesh.Normals[VertIndex];
+			FVector& OutPosition = OutPositions[VertIndex];
+			FVector& OutNormal = OutNormals[VertIndex];
+			switch (InMesh.BoneData[VertIndex].NumInfluences)
+			{
+			default: break;
+			case 8:
+			{
+				const FMatrix& BoneMatrix = BoneMatrices[BoneMap[BoneIndices[7]]];
+				const float Weight = BoneWeights[7];
+
+				OutPosition += BoneMatrix.TransformPosition(RefParticle) * Weight;
+				OutNormal += BoneMatrix.TransformVector(RefNormal) * Weight;
+			}
+			case 7:
+			{
+				const FMatrix& BoneMatrix = BoneMatrices[BoneMap[BoneIndices[6]]];
+				const float Weight = BoneWeights[6];
+
+				OutPosition += BoneMatrix.TransformPosition(RefParticle) * Weight;
+				OutNormal += BoneMatrix.TransformVector(RefNormal) * Weight;
+			}
+			case 6:
+			{
+				const FMatrix& BoneMatrix = BoneMatrices[BoneMap[BoneIndices[5]]];
+				const float Weight = BoneWeights[5];
+
+				OutPosition += BoneMatrix.TransformPosition(RefParticle) * Weight;
+				OutNormal += BoneMatrix.TransformVector(RefNormal) * Weight;
+			}
+			case 5:
+			{
+				const FMatrix& BoneMatrix = BoneMatrices[BoneMap[BoneIndices[4]]];
+				const float Weight = BoneWeights[4];
+
+				OutPosition += BoneMatrix.TransformPosition(RefParticle) * Weight;
+				OutNormal += BoneMatrix.TransformVector(RefNormal) * Weight;
+			}
+			case 4:
+			{
+				const FMatrix& BoneMatrix = BoneMatrices[BoneMap[BoneIndices[3]]];
+				const float Weight = BoneWeights[3];
+
+				OutPosition += BoneMatrix.TransformPosition(RefParticle) * Weight;
+				OutNormal += BoneMatrix.TransformVector(RefNormal) * Weight;
+			}
+			case 3:
+			{
+				const FMatrix& BoneMatrix = BoneMatrices[BoneMap[BoneIndices[2]]];
+				const float Weight = BoneWeights[2];
+
+				OutPosition += BoneMatrix.TransformPosition(RefParticle) * Weight;
+				OutNormal += BoneMatrix.TransformVector(RefNormal) * Weight;
+			}
+			case 2:
+			{
+				const FMatrix& BoneMatrix = BoneMatrices[BoneMap[BoneIndices[1]]];
+				const float Weight = BoneWeights[1];
+
+				OutPosition += BoneMatrix.TransformPosition(RefParticle) * Weight;
+				OutNormal += BoneMatrix.TransformVector(RefNormal) * Weight;
+			}
+			case 1:
+			{
+				const FMatrix& BoneMatrix = BoneMatrices[BoneMap[BoneIndices[0]]];
+				const float Weight = BoneWeights[0];
+
+				OutPosition += BoneMatrix.TransformPosition(RefParticle) * Weight;
+				OutNormal += BoneMatrix.TransformVector(RefNormal) * Weight;
+			}
+			}
+
+			OutPosition = RootBoneTransformInternal.InverseTransformPosition(OutPosition);
+			OutNormal = RootBoneTransformInternal.InverseTransformVector(OutNormal);
+
+			if (OutNormal.SizeSquared() > SMALL_NUMBER)
+			{
+				OutNormal = OutNormal.GetUnsafeNormal();
+			}
+		}
+	}
+
+	/** 
+	 * Gets the best match triangle for a specified position from the triangles in Mesh.
+	 * Performs no validation on the incoming mesh data, the mesh data should be verified
+	 * to be valid before using this function
+	 */
+	int32 GetBestTriangleBaseIndex(const ClothMeshDesc& Mesh, const FVector& Position)
+	{
+		float MinimumDistanceSq = MAX_flt;
+		int32 ClosestBaseIndex = INDEX_NONE;
+
+		const int32 NumTriangles = Mesh.Indices.Num() / 3;
+		for(int32 TriIdx = 0; TriIdx < NumTriangles; ++TriIdx)
+		{
+			int32 TriBaseIdx = TriIdx * 3;
+
+			const uint32 IA = Mesh.Indices[TriBaseIdx + 0];
+			const uint32 IB = Mesh.Indices[TriBaseIdx + 1];
+			const uint32 IC = Mesh.Indices[TriBaseIdx + 2];
+
+			const FVector& A = Mesh.Positions[IA];
+			const FVector& B = Mesh.Positions[IB];
+			const FVector& C = Mesh.Positions[IC];
+
+			FVector PointOnTri = FMath::ClosestPointOnTriangleToPoint(Position, A, B, C);
+			float DistSq = (PointOnTri - Position).SizeSquared();
+
+			if(DistSq < MinimumDistanceSq)
+			{
+				MinimumDistanceSq = DistSq;
+				ClosestBaseIndex = TriBaseIdx;
+			}
+		}
+
+		return ClosestBaseIndex;
+	}
+
+	void GenerateMeshToMeshSkinningData(TArray<FMeshToMeshVertData>& OutSkinningData, const ClothMeshDesc& TargetMesh, const TArray<FVector>* TargetTangents, const ClothMeshDesc& SourceMesh)
+	{
+		if(!TargetMesh.HasValidMesh())
+		{
+			UE_LOG(LogClothingMeshUtils, Warning, TEXT("Failed to generate mesh to mesh skinning data. Invalid Target Mesh."));
+			return;
+		}
+
+		if(!SourceMesh.HasValidMesh())
+		{
+			UE_LOG(LogClothingMeshUtils, Warning, TEXT("Failed to generate mesh to mesh skinning data. Invalid Source Mesh."));
+			return;
+		}
+
+		const int32 NumMesh0Verts = TargetMesh.Positions.Num();
+		const int32 NumMesh0Normals = TargetMesh.Normals.Num();
+		const int32 NumMesh0Tangents = TargetTangents ? TargetTangents->Num() : 0;
+
+		const int32 NumMesh1Verts = SourceMesh.Positions.Num();
+		const int32 NumMesh1Normals = SourceMesh.Normals.Num();
+		const int32 NumMesh1Indices = SourceMesh.Indices.Num();
+
+		// Check we have properly formed triangles
+		check(NumMesh1Indices % 3 == 0);
+
+		const int32 NumMesh1Triangles = NumMesh1Indices / 3;
+
+		// Check mesh data to make sure we have the same number of each element
+		if(NumMesh0Verts != NumMesh0Normals || (TargetTangents && NumMesh0Tangents != NumMesh0Verts))
+		{
+			UE_LOG(LogClothingMeshUtils, Warning, TEXT("Can't generate mesh to mesh skinning data, Mesh0 data is missing verts."));
+			return;
+		}
+
+		if(NumMesh1Verts != NumMesh1Normals)
+		{
+			UE_LOG(LogClothingMeshUtils, Warning, TEXT("Can't generate mesh to mesh skinning data, Mesh1 data is missing verts."));
+			return;
+		}
+
+		OutSkinningData.Reserve(NumMesh0Verts);
+
+		// For all mesh0 verts
+		for(int32 VertIdx0 = 0; VertIdx0 < NumMesh0Verts; ++VertIdx0)
+		{
+			OutSkinningData.AddZeroed();
+			FMeshToMeshVertData& SkinningData = OutSkinningData.Last();
+
+			const FVector& VertPosition = TargetMesh.Positions[VertIdx0];
+			const FVector& VertNormal = TargetMesh.Normals[VertIdx0];
+
+			FVector VertTangent;
+			if(TargetTangents)
+			{
+				VertTangent = (*TargetTangents)[VertIdx0];
+			}
+			else
+			{
+				FVector Tan0, Tan1;
+				VertNormal.FindBestAxisVectors(Tan0, Tan1);
+				VertTangent = Tan0;
+			}
+
+			int32 ClosestTriangleBaseIdx = GetBestTriangleBaseIndex(SourceMesh, VertPosition);
+
+			check(ClosestTriangleBaseIdx != INDEX_NONE);
+
+			const FVector& A = SourceMesh.Positions[SourceMesh.Indices[ClosestTriangleBaseIdx]];
+			const FVector& B = SourceMesh.Positions[SourceMesh.Indices[ClosestTriangleBaseIdx + 1]];
+			const FVector& C = SourceMesh.Positions[SourceMesh.Indices[ClosestTriangleBaseIdx + 2]];
+
+			const FVector& NA = SourceMesh.Normals[SourceMesh.Indices[ClosestTriangleBaseIdx]];
+			const FVector& NB = SourceMesh.Normals[SourceMesh.Indices[ClosestTriangleBaseIdx + 1]];
+			const FVector& NC = SourceMesh.Normals[SourceMesh.Indices[ClosestTriangleBaseIdx + 2]];
+
+			// Before generating the skinning data we need to check for a degenerate triangle.
+			// If we find _any_ degenerate triangles we will notify and fail to generate the skinning data
+			const FVector TriNormal = FVector::CrossProduct(B - A, C - A);
+			if(TriNormal.SizeSquared() < SMALL_NUMBER)
+			{
+				// Failed, we have 2 identical vertices
+				OutSkinningData.Reset();
+
+				// Log and toast
+				FText Error = FText::Format(LOCTEXT("DegenerateTriangleError", "Failed to generate skinning data, found conincident vertices in triangle A={0} B={1} C={2}"), FText::FromString(A.ToString()), FText::FromString(B.ToString()), FText::FromString(C.ToString()));
+
+				UE_LOG(LogClothingMeshUtils, Warning, TEXT("%s"), *Error.ToString());
+
+#if WITH_EDITOR
+				FNotificationInfo Info(Error);
+				Info.ExpireDuration = 5.0f;
+				FSlateNotificationManager::Get().AddNotification(Info);
+#endif
+				return;
+			}
+
+			SkinningData.PositionBaryCoordsAndDist = GetPointBaryAndDist(A, B, C, NA, NB, NC, VertPosition);
+			SkinningData.NormalBaryCoordsAndDist = GetPointBaryAndDist(A, B, C, NA, NB, NC, VertPosition + VertNormal);
+			SkinningData.TangentBaryCoordsAndDist = GetPointBaryAndDist(A, B, C, NA, NB, NC, VertPosition + VertTangent);
+			SkinningData.SourceMeshVertIndices[0] = SourceMesh.Indices[ClosestTriangleBaseIdx];
+			SkinningData.SourceMeshVertIndices[1] = SourceMesh.Indices[ClosestTriangleBaseIdx + 1];
+			SkinningData.SourceMeshVertIndices[2] = SourceMesh.Indices[ClosestTriangleBaseIdx + 2];
+			SkinningData.SourceMeshVertIndices[3] = 0;
+		}
+	}
+
+	FVector4 GetPointBaryAndDist(const FVector& A, const FVector& B, const FVector& C, const FVector& NA, const FVector& NB, const FVector& NC, const FVector& Point)
+	{
+		FPlane TrianglePlane(A, B, C);
+		const FVector PointOnTriPlane = FVector::PointPlaneProject(Point, TrianglePlane);
+		const FVector BaryCoords = FMath::ComputeBaryCentric2D(PointOnTriPlane, A, B, C);
+		const FVector NormalAtPoint = TrianglePlane;
+		FVector TriPointToVert = Point - PointOnTriPlane;
+		TriPointToVert = TriPointToVert.ProjectOnTo(NormalAtPoint);
+		float Dist = TriPointToVert.Size();
+
+		float Sign = TrianglePlane.PlaneDot(Point) < 0.0f ? -1.0f : 1.0f;
+
+		return FVector4(BaryCoords, TrianglePlane.PlaneDot(Point));
+	}
+
+	void GenerateEmbeddedPositions(const ClothMeshDesc& SourceMesh, TArrayView<const FVector> Positions, TArray<FVector4>& OutEmbeddedPositions, TArray<int32>& OutSourceIndices)
+	{
+		if(!SourceMesh.HasValidMesh())
+		{
+			// No valid source mesh
+			return;
+		}
+
+		const int32 NumPositions = Positions.Num();
+
+		OutEmbeddedPositions.Reset();
+		OutEmbeddedPositions.AddUninitialized(NumPositions);
+
+		OutSourceIndices.Reset(NumPositions * 3);
+
+		for(int32 PositionIndex = 0 ; PositionIndex < NumPositions ; ++PositionIndex)
+		{
+			const FVector& Position = Positions[PositionIndex];
+
+			int32 TriBaseIndex = GetBestTriangleBaseIndex(SourceMesh, Position);
+
+			const int32 IA = SourceMesh.Indices[TriBaseIndex];
+			const int32 IB = SourceMesh.Indices[TriBaseIndex + 1];
+			const int32 IC = SourceMesh.Indices[TriBaseIndex + 2];
+
+			const FVector& A = SourceMesh.Positions[IA];
+			const FVector& B = SourceMesh.Positions[IB];
+			const FVector& C = SourceMesh.Positions[IC];
+
+			const FVector& NA = SourceMesh.Normals[IA];
+			const FVector& NB = SourceMesh.Normals[IB];
+			const FVector& NC = SourceMesh.Normals[IC];
+
+			OutEmbeddedPositions[PositionIndex] = GetPointBaryAndDist(A, B, C, NA, NB, NC, Position);
+			OutSourceIndices.Add(IA);
+			OutSourceIndices.Add(IB);
+			OutSourceIndices.Add(IC);
+		}
+	}
+
+	void FVertexParameterMapper::Map(TArrayView<const float> Source, TArray<float>& Dest)
+	{
+		Map(Source, Dest, [](FVector Bary, float A, float B, float C)
+		{
+			return Bary.X * A + Bary.Y * B + Bary.Z * C;
+		});
+	}
+
+}
+
+#undef LOCTEXT_NAMESPACE
