Copying //UE4/Dev-Physics to Dev-Main (//UE4/Dev-Main) [at] 9098278

diff --git a/Engine/Source/Runtime/ClothingSystemRuntimeNv/Public/Assets/ClothingAsset.h b/Engine/Source/Runtime/ClothingSystemRuntimeNv/Public/Assets/ClothingAsset.h
new file mode 100644
--- /dev/null
+++ b/Engine/Source/Runtime/ClothingSystemRuntimeNv/Public/Assets/ClothingAsset.h
@@ -0,0 +1,610 @@
+// Copyright 1998-2019 Epic Games, Inc. All Rights Reserved.
+
+#pragma once
+
+#include "ClothingAssetBase.h"
+
+#include "ClothCollisionData.h"
+#include "ClothingSystemRuntimeTypes.h"
+#include "GPUSkinPublicDefs.h"
+#include "SkeletalMeshTypes.h"
+
+#include "ClothingAsset.generated.h"
+
+class UClothingAsset;
+class UPhysicsAsset;
+struct FClothPhysicalMeshData;
+struct FSkelMeshSection;
+class UClothingSimulationInteractor;
+
+namespace nvidia
+{
+	namespace apex
+	{
+		class ClothingAsset;
+	}
+}
+
+DECLARE_LOG_CATEGORY_EXTERN(LogClothingAsset, Log, All);
+
+/** The possible targets for a physical mesh mask */
+UENUM()
+enum class MaskTarget_PhysMesh : uint8
+{
+	None,
+	MaxDistance,
+	BackstopDistance,
+	BackstopRadius,
+	AnimDriveMultiplier
+};
+
+/** 
+ * A mask is simply some storage for a physical mesh parameter painted onto clothing.
+ * Used in the editor for users to paint onto and then target to a parameter, which
+ * is then later applied to a phys mesh
+ */
+USTRUCT()
+struct CLOTHINGSYSTEMRUNTIMENV_API FClothParameterMask_PhysMesh
+{
+	GENERATED_BODY();
+
+	FClothParameterMask_PhysMesh()
+		: MaskName(NAME_None)
+		, CurrentTarget(MaskTarget_PhysMesh::None)
+		, MaxValue_DEPRECATED(0.0)
+		, MinValue_DEPRECATED(100.0)
+		, bEnabled(false)
+	{}
+
+	/** 
+	 * Initialize the mask based on the specified mesh data
+	 * @param InMeshData the mesh to initialize against
+	 */
+	void Initialize(const FClothPhysicalMeshData& InMeshData);
+
+	/** 
+	 * Copies the specified parameter from a physical mesh
+	 * @param InMeshData The mesh to copy from
+	 * @param InTarget The target parameter to copy
+	 */
+	void CopyFromPhysMesh(const FClothPhysicalMeshData& InMeshData, MaskTarget_PhysMesh InTarget);
+
+	/** 
+	 * Set a value in the mask
+	 * @param InVertexIndex the value/vertex index to set
+	 * @param InValue the value to set
+	 */
+	void SetValue(int32 InVertexIndex, float InValue);
+
+	/** 
+	 * Get a value from the mask
+	 * @param InVertexIndex the value/vertex index to retrieve
+	 */
+	float GetValue(int32 InVertexIndex) const;
+	
+	/** 
+	* Read only version of the array holding the mask values
+	*/
+	const TArray<float>& GetValueArray() const;
+
+	/** Calculates Min/Max values based on values. */
+	void CalcRanges(float& MinValue, float& MaxValue);
+
+	/** 
+	 * Apply the mask to the provided mesh (if compatible)
+	 * @param InTargetMesh the targeted mesh
+	 */
+	void Apply(FClothPhysicalMeshData& InTargetMesh);
+
+	/** Name of the mask, mainly for users to differentiate */
+	UPROPERTY()
+	FName MaskName;
+
+	/** The currently targeted parameter for the mask */
+	UPROPERTY()
+	MaskTarget_PhysMesh CurrentTarget;
+
+	/** The maximum value currently in the mask value array */
+	UPROPERTY()
+	float MaxValue_DEPRECATED;
+
+	/** The maximum value currently in the mask value array */
+	UPROPERTY()
+	float MinValue_DEPRECATED;
+
+	/** The actual values stored in the mask */
+	UPROPERTY()
+	TArray<float> Values;
+
+	/** Whether this mask is enabled and able to effect final mesh values */
+	UPROPERTY()
+	bool bEnabled;
+};
+
+// Bone data for a vertex
+USTRUCT()
+struct FClothVertBoneData
+{
+	GENERATED_BODY()
+
+	FClothVertBoneData()
+		: NumInfluences(0)
+	{
+		FMemory::Memset(BoneIndices, (uint8)INDEX_NONE, sizeof(BoneIndices));
+		FMemory::Memset(BoneWeights, 0, sizeof(BoneWeights));
+	}
+
+	// Number of influences for this vertex.
+	UPROPERTY()
+	int32 NumInfluences;
+
+	// Up to MAX_TOTAL_INFLUENCES bone indices that this vert is weighted to
+	UPROPERTY()
+	uint16 BoneIndices[MAX_TOTAL_INFLUENCES];
+
+	// The weights for each entry in BoneIndices
+	UPROPERTY()
+	float BoneWeights[MAX_TOTAL_INFLUENCES];
+};
+
+/**
+ *	Physical mesh data created during asset import or created from a skeletal mesh
+ */
+USTRUCT()
+struct CLOTHINGSYSTEMRUNTIMENV_API FClothPhysicalMeshData
+{
+	GENERATED_BODY()
+
+	void Reset(const int32 InNumVerts);
+
+	// Clear out any target properties in this physical mesh
+	void ClearParticleParameters();
+
+	// Whether the mesh uses backstops
+	bool HasBackStops() const;
+
+	// Whether the mesh uses anim drives
+	bool HasAnimDrive() const;
+
+	// Positions of each simulation vertex
+	UPROPERTY(EditAnywhere, Category = SimMesh)
+	TArray<FVector> Vertices;
+
+	// Normal at each vertex
+	UPROPERTY(EditAnywhere, Category = SimMesh)
+	TArray<FVector> Normals;
+
+#if WITH_EDITORONLY_DATA
+	// Color at each vertex
+	UPROPERTY(EditAnywhere, Category = SimMesh)
+	TArray<FColor> VertexColors;
+#endif // WITH_EDITORONLY_DATA
+
+	// Indices of the simulation mesh triangles
+	UPROPERTY(EditAnywhere, Category = SimMesh)
+	TArray<uint32> Indices;
+
+	// The distance that each vertex can move away from its reference (skinned) position
+	UPROPERTY(EditAnywhere, Category = SimMesh)
+	TArray<float> MaxDistances;
+
+	// Distance along the plane of the surface that the particles can travel (separation constraint)
+	UPROPERTY(EditAnywhere, Category = SimMesh)
+	TArray<float> BackstopDistances;
+
+	// Radius of movement to allow for backstop movement
+	UPROPERTY(EditAnywhere, Category = SimMesh)
+	TArray<float> BackstopRadiuses;
+
+	// Strength of anim drive per-particle (spring driving particle back to skinned location
+	UPROPERTY(EditAnywhere, Category = SimMesh)
+	TArray<float> AnimDriveMultipliers;
+
+	// Inverse mass for each vertex in the physical mesh
+	UPROPERTY(EditAnywhere, Category = SimMesh)
+	TArray<float> InverseMasses;
+
+	// Indices and weights for each vertex, used to skin the mesh to create the reference pose
+	UPROPERTY(EditAnywhere, Category = SimMesh)
+	TArray<FClothVertBoneData> BoneData;
+
+	// Maximum number of bone weights of any vetex
+	UPROPERTY(EditAnywhere, Category = SimMesh)
+	int32 MaxBoneWeights;
+
+	// Number of fixed verts in the simulation mesh (fixed verts are just skinned and do not simulate)
+	UPROPERTY(EditAnywhere, Category = SimMesh)
+	int32 NumFixedVerts;
+
+	// Valid indices to use for self collisions (reduced set of Indices)
+	UPROPERTY(EditAnywhere, Category = SimMesh)
+	TArray<uint32> SelfCollisionIndices;
+};
+
+USTRUCT()
+struct CLOTHINGSYSTEMRUNTIMENV_API FClothLODData
+{
+	GENERATED_BODY()
+
+	// Raw phys mesh data
+	UPROPERTY(EditAnywhere, Category = SimMesh)
+	FClothPhysicalMeshData PhysicalMeshData;
+
+	// Collision primitive and covex data for clothing collisions
+	UPROPERTY(EditAnywhere, Category = Collision)
+	FClothCollisionData CollisionData;
+
+#if WITH_EDITORONLY_DATA
+	// Parameter masks defining the physics mesh masked data
+	UPROPERTY(EditAnywhere, Category = Masks)
+	TArray<FClothParameterMask_PhysMesh> ParameterMasks;
+
+	// Get all available parameter masks for the specified target
+	void GetParameterMasksForTarget(const MaskTarget_PhysMesh& InTarget, TArray<FClothParameterMask_PhysMesh*>& OutMasks);
+
+#endif
+
+	// Skinning data for transitioning from a higher detail LOD to this one
+	TArray<FMeshToMeshVertData> TransitionUpSkinData;
+
+	// Skinning data for transitioning from a lower detail LOD to this one
+	TArray<FMeshToMeshVertData> TransitionDownSkinData;
+
+	bool Serialize(FArchive& Ar)
+	{
+		UScriptStruct* Struct = FClothLODData::StaticStruct();
+
+		// Serialize normal tagged data
+		if (!Ar.IsCountingMemory())
+		{
+			Struct->SerializeTaggedProperties(Ar, (uint8*)this, Struct, nullptr);
+		}
+
+		// Serialize the mesh to mesh data (not a USTRUCT)
+		Ar	<< TransitionUpSkinData
+			<< TransitionDownSkinData;
+
+		return true;
+	}
+};
+
+template<>
+struct TStructOpsTypeTraits<FClothLODData> : public TStructOpsTypeTraitsBase2<FClothLODData>
+{
+	enum
+	{
+		WithSerializer = true,
+	};
+};
+
+// Container for a constraint setup, these can be horizontal, vertical, shear and bend
+USTRUCT()
+struct FClothConstraintSetup
+{
+	GENERATED_BODY()
+
+	FClothConstraintSetup()
+		: Stiffness(1.0f)
+		, StiffnessMultiplier(1.0f)
+		, StretchLimit(1.0f)
+		, CompressionLimit(1.0f)
+	{}
+
+	// How stiff this constraint is, this affects how closely it will follow the desired position
+	UPROPERTY(EditAnywhere, Category=Constraint)
+	float Stiffness;
+
+	// A multiplier affecting the above value
+	UPROPERTY(EditAnywhere, Category = Constraint)
+	float StiffnessMultiplier;
+
+	// The hard limit on how far this constarint can stretch
+	UPROPERTY(EditAnywhere, Category = Constraint)
+	float StretchLimit;
+
+	// The hard limit on how far this constraint can compress
+	UPROPERTY(EditAnywhere, Category = Constraint)
+	float CompressionLimit;
+};
+
+UENUM()
+enum class EClothingWindMethod : uint8
+{
+	// Use legacy wind mode, where accelerations are modified directly by the simulation
+	// with no regard for drag or lift
+	Legacy,
+
+	// Use updated wind calculation for NvCloth based solved taking into account
+	// drag and lift, this will require those properties to be correctly set in
+	// the clothing configuration
+	Accurate
+};
+
+/** Holds initial, asset level config for clothing actors. */
+USTRUCT()
+struct FClothConfig
+{
+	GENERATED_BODY()
+
+	FClothConfig()
+		: WindMethod(EClothingWindMethod::Legacy)
+		, SelfCollisionRadius(0.0f)
+		, SelfCollisionStiffness(0.0f)
+		, SelfCollisionCullScale(1.0f)
+		, Damping(0.4f)
+		, Friction(0.1f)
+		, WindDragCoefficient(0.02f/100.0f)
+		, WindLiftCoefficient(0.02f/100.0f)
+		, LinearDrag(0.2f)
+		, AngularDrag(0.2f)
+		, LinearInertiaScale(1.0f)
+		, AngularInertiaScale(1.0f)
+		, CentrifugalInertiaScale(1.0f)
+		, SolverFrequency(120.0f)
+		, StiffnessFrequency(100.0f)
+		, GravityScale(1.0f)
+		, GravityOverride(FVector::ZeroVector)
+		, bUseGravityOverride(false)
+		, TetherStiffness(1.0f)
+		, TetherLimit(1.0f)
+		, CollisionThickness(1.0f)
+		, AnimDriveSpringStiffness(1.0f)
+		, AnimDriveDamperStiffness(1.0f)
+	{}
+
+	bool HasSelfCollision() const;
+
+	// How wind should be processed, Accurate uses drag and lift to make the cloth react differently, legacy applies similar forces to all clothing without drag and lift (similar to APEX)
+	UPROPERTY(EditAnywhere, Category = ClothConfig)
+	EClothingWindMethod WindMethod;
+
+	// Constraint data for vertical constraints
+	UPROPERTY(EditAnywhere, Category = ClothConfig)
+	FClothConstraintSetup VerticalConstraintConfig;
+
+	// Constraint data for horizontal constraints
+	UPROPERTY(EditAnywhere, Category = ClothConfig)
+	FClothConstraintSetup HorizontalConstraintConfig;
+
+	// Constraint data for bend constraints
+	UPROPERTY(EditAnywhere, Category = ClothConfig)
+	FClothConstraintSetup BendConstraintConfig;
+
+	// Constraint data for shear constraints
+	UPROPERTY(EditAnywhere, Category = ClothConfig)
+	FClothConstraintSetup ShearConstraintConfig;
+
+	// Size of self collision spheres centered on each vert
+	UPROPERTY(EditAnywhere, Category = ClothConfig)
+	float SelfCollisionRadius;
+
+	// Stiffness of the spring force that will resolve self collisions
+	UPROPERTY(EditAnywhere, Category = ClothConfig)
+	float SelfCollisionStiffness;
+
+	/** 
+	 * Scale to use for the radius of the culling checks for self collisions.
+	 * Any other self collision body within the radius of this check will be culled.
+	 * This helps performance with higher resolution meshes by reducing the number
+	 * of colliding bodies within the cloth. Reducing this will have a negative
+	 * effect on performance!
+	 */
+	UPROPERTY(EditAnywhere, Category = ClothConfig, meta = (UIMin="0", ClampMin="0"))
+	float SelfCollisionCullScale;
+
+	// Damping of particle motion per-axis
+	UPROPERTY(EditAnywhere, Category = ClothConfig)
+	FVector Damping;
+
+	// Friction of the surface when colliding
+	UPROPERTY(EditAnywhere, Category = ClothConfig)
+	float Friction;
+
+	// Drag coefficient for wind calculations, higher values mean wind has more lateral effect on cloth
+	UPROPERTY(EditAnywhere, Category = ClothConfig)
+	float WindDragCoefficient;
+
+	// Lift coefficient for wind calculations, higher values make cloth rise easier in wind
+	UPROPERTY(EditAnywhere, Category = ClothConfig)
+	float WindLiftCoefficient;
+
+	// Drag applied to linear particle movement per-axis
+	UPROPERTY(EditAnywhere, Category = ClothConfig)
+	FVector LinearDrag;
+
+	// Drag applied to angular particle movement, higher values should limit material bending (per-axis)
+	UPROPERTY(EditAnywhere, Category = ClothConfig)
+	FVector AngularDrag;
+
+	// Scale for linear particle inertia, how much movement should translate to linear motion (per-axis)
+	UPROPERTY(EditAnywhere, Category = ClothConfig, meta = (UIMin="0", UIMax="1", ClampMin="0", ClampMax="1"))
+	FVector LinearInertiaScale;
+
+	// Scale for angular particle inertia, how much movement should translate to angular motion (per-axis)
+	UPROPERTY(EditAnywhere, Category = ClothConfig, meta = (UIMin = "0", UIMax = "1", ClampMin = "0", ClampMax = "1"))
+	FVector AngularInertiaScale;
+
+	// Scale for centrifugal particle inertia, how much movement should translate to angular motion (per-axis)
+	UPROPERTY(EditAnywhere, Category = ClothConfig, meta = (UIMin = "0", UIMax = "1", ClampMin = "0", ClampMax = "1"))
+	FVector CentrifugalInertiaScale;
+
+	// Frequency of the position solver, lower values will lead to stretchier, bouncier cloth
+	UPROPERTY(EditAnywhere, Category = ClothConfig, meta = (UIMin = "30", UIMax = "240", ClampMin = "30", ClampMax = "1000"))
+	float SolverFrequency;
+
+	// Frequency for stiffness calculations, lower values will degrade stiffness of constraints
+	UPROPERTY(EditAnywhere, Category = ClothConfig)
+	float StiffnessFrequency;
+
+	// Scale of gravity effect on particles
+	UPROPERTY(EditAnywhere, Category = ClothConfig, meta = (EditCondition = "!bUseGravityOverride"))
+	float GravityScale;
+
+	// Direct gravity override value
+	UPROPERTY(EditAnywhere, Category = ClothConfig, meta = (EditCondition = "bUseGravityOverride"))
+	FVector GravityOverride;
+
+	/** Use gravity override value vs gravity scale */
+	UPROPERTY(EditAnywhere, Category = ClothConfig, meta = (InlineEditConditionToggle))
+	bool bUseGravityOverride;
+
+	// Scale for stiffness of particle tethers between each other
+	UPROPERTY(EditAnywhere, Category = ClothConfig)
+	float TetherStiffness;
+
+	// Scale for the limit of particle tethers (how far they can separate)
+	UPROPERTY(EditAnywhere, Category = ClothConfig)
+	float TetherLimit;
+
+	// 'Thickness' of the simulated cloth, used to adjust collisions
+	UPROPERTY(EditAnywhere, Category = ClothConfig)
+	float CollisionThickness;
+
+	// Default spring stiffness for anim drive if an anim drive is in use
+	UPROPERTY(EditAnywhere, Category = ClothConfig)
+	float AnimDriveSpringStiffness;
+
+	// Default damper stiffness for anim drive if an anim drive is in use
+	UPROPERTY(EditAnywhere, Category = ClothConfig)
+	float AnimDriveDamperStiffness;
+};
+
+namespace ClothingAssetUtils
+{
+	// Helper struct to hold binding information on a clothing asset, used to enumerate all of the
+	// bindings on a skeletal mesh with GetMeshClothingAssetBindings below.
+	struct FClothingAssetMeshBinding
+	{
+		UClothingAsset* Asset;
+		int32 LODIndex;
+		int32 SectionIndex;
+		int32 AssetInternalLodIndex;
+	};
+
+	/**
+	 * Given a skeletal mesh, find all of the currently bound clothing assets and their binding information
+	 * @param InSkelMesh - The skeletal mesh to search
+	 * @param OutBindings - The list of bindings to write to
+	 */
+	void CLOTHINGSYSTEMRUNTIMENV_API GetMeshClothingAssetBindings(USkeletalMesh* InSkelMesh, TArray<FClothingAssetMeshBinding>& OutBindings);
+	
+	// Similar to above, but only inspects the specified LOD
+	void CLOTHINGSYSTEMRUNTIMENV_API GetMeshClothingAssetBindings(USkeletalMesh* InSkelMesh, TArray<FClothingAssetMeshBinding>& OutBindings, int32 InLodIndex);
+
+#if WITH_EDITOR
+	// Clears the clothing tracking struct of a section.
+	CLOTHINGSYSTEMRUNTIMENV_API void ClearSectionClothingData(FSkelMeshSection& InSection);
+#endif
+}
+
+/**
+ * Custom data wrapper for clothing assets
+ * If writing a new clothing asset importer, creating a new derived custom data is how to store importer (and possibly simulation)
+ * data that importer will create. This needs to be set to the CustomData member on the asset your factory creates.
+ * Testing whether a UClothingAsset was made from a custom plugin can be achieved with
+ * if(AssetPtr->CustomData->IsA(UMyCustomData::StaticClass()))
+ */
+UCLASS(abstract, MinimalAPI)
+class UClothingAssetCustomData : public UObject
+{
+	GENERATED_BODY()
+
+public:
+	virtual void BindToSkeletalMesh(USkeletalMesh* InSkelMesh, int32 InMeshLodIndex, int32 InSectionIndex, int32 InAssetLodIndex)
+	{}
+};
+
+UCLASS(hidecategories = Object, BlueprintType)
+class CLOTHINGSYSTEMRUNTIMENV_API UClothingAsset : public UClothingAssetBase
+{
+	GENERATED_BODY()
+
+public:
+
+	UClothingAsset(const FObjectInitializer& ObjectInitializer);
+
+	// UClothingAssetBase Interface ////////////////////////////////////////////
+	virtual void RefreshBoneMapping(USkeletalMesh* InSkelMesh) override;
+#if WITH_EDITOR
+	virtual bool BindToSkeletalMesh(USkeletalMesh* InSkelMesh, int32 InMeshLodIndex, int32 InSectionIndex, int32 InAssetLodIndex, bool bCallPostEditChange=true) override;
+	virtual void UnbindFromSkeletalMesh(USkeletalMesh* InSkelMesh) override;
+	virtual void UnbindFromSkeletalMesh(USkeletalMesh* InSkelMesh, int32 InMeshLodIndex) override;
+	virtual void InvalidateCachedData() override;
+	virtual bool IsValidLod(int32 InLodIndex) override;
+	virtual int32 GetNumLods() override;
+	// End UClothingAssetBase Interface ////////////////////////////////////////
+
+	/**
+	*	Builds the LOD transition data
+	*	When we transition between LODs we skin the incoming mesh to the outgoing mesh
+	*	in exactly the same way the render mesh is skinned to create a smooth swap
+	*/
+	void BuildLodTransitionData();
+
+	/** 
+	 * Applies the painted parameter masks to the final parameters in the physical mesh
+	 */
+	void ApplyParameterMasks();
+#endif
+
+	// UObject Interface //////////////////////////////////////////////////////
+	
+	virtual void PostLoad() override;
+	virtual void Serialize(FArchive& Ar) override;
+
+#if WITH_EDITOR
+	virtual void PostEditChangeChainProperty(struct FPropertyChangedChainEvent& PropertyChangedEvent) override;
+#endif
+
+	// End UObject Interface //////////////////////////////////////////////////
+
+#if WITH_EDITOR
+	void ReregisterComponentsUsingClothing();
+	void ForEachInteractorUsingClothing(TFunction<void (UClothingSimulationInteractor*)> Func);
+#endif
+
+	/** 
+	 * Builds self collision data
+	 * This is a subset of actual verts that need to be considered for self collision such that we cover
+	 * as much of the mesh as possible in collision bodies so that nothing can fall through
+	 */
+	virtual void BuildSelfCollisionData() override;
+
+	// Calculates the prefered root bone for the simulation
+	void CalculateReferenceBoneIndex();
+
+	// The physics asset to extract collisions from when building a simulation
+	UPROPERTY(EditAnywhere, Category = Config)
+	UPhysicsAsset* PhysicsAsset;
+
+	// Configuration of the cloth, contains all the parameters for how the clothing behaves
+	UPROPERTY(EditAnywhere, Category = Config)
+	FClothConfig ClothConfig;
+
+	// The actual asset data, listed by LOD
+	UPROPERTY()
+	TArray<FClothLODData> LodData;
+
+	// Tracks which clothing LOD each skel mesh LOD corresponds to (LodMap[SkelLod]=ClothingLod)
+	UPROPERTY()
+	TArray<int32> LodMap;
+
+	// List of bones this asset uses inside its parent mesh
+	UPROPERTY()
+	TArray<FName> UsedBoneNames;
+
+	// List of the indices for the bones in UsedBoneNames, used for remapping
+	UPROPERTY()
+	TArray<int32> UsedBoneIndices;
+
+	// Bone to treat as the root of the simulation space
+	UPROPERTY()
+	int32 ReferenceBoneIndex;
+
+	/** Custom data applied by the importer depending on where the asset was imported from */
+	UPROPERTY()
+	UClothingAssetCustomData* CustomData;
+
+private:
+
+};
diff --git a/Engine/Source/Runtime/ClothingSystemRuntimeNv/Public/ClothingActor.h b/Engine/Source/Runtime/ClothingSystemRuntimeNv/Public/ClothingActor.h
new file mode 100644
--- /dev/null
+++ b/Engine/Source/Runtime/ClothingSystemRuntimeNv/Public/ClothingActor.h
@@ -0,0 +1,14 @@
+// Copyright 1998-2019 Epic Games, Inc. All Rights Reserved.
+
+#pragma once
+
+class UClothingAsset;
+
+// Base clothing actor just needs a link back to its parent asset, everything else defined by derived simulation
+class FClothingActorBase
+{
+public:
+
+	// Link back to parent asset
+	UClothingAsset* AssetCreatedFrom;
+};
diff --git a/Engine/Source/Runtime/ClothingSystemRuntimeNv/Public/ClothingSimulationFactoryNv.h b/Engine/Source/Runtime/ClothingSystemRuntimeNv/Public/ClothingSimulationFactoryNv.h
new file mode 100644
--- /dev/null
+++ b/Engine/Source/Runtime/ClothingSystemRuntimeNv/Public/ClothingSimulationFactoryNv.h
@@ -0,0 +1,22 @@
+// Copyright 1998-2019 Epic Games, Inc. All Rights Reserved.
+
+#pragma once
+
+#include "ClothingSimulation.h"
+#include "ClothingSimulationFactory.h"
+
+#include "ClothingSimulationFactoryNv.generated.h"
+
+UCLASS()
+class CLOTHINGSYSTEMRUNTIMENV_API UClothingSimulationFactoryNv final : public UClothingSimulationFactory
+{
+	GENERATED_BODY()
+public:
+
+	virtual IClothingSimulation* CreateSimulation() override;
+	virtual void DestroySimulation(IClothingSimulation* InSimulation) override;
+	virtual bool SupportsAsset(UClothingAssetBase* InAsset) override;
+
+	virtual bool SupportsRuntimeInteraction() override;
+	virtual UClothingSimulationInteractor* CreateInteractor() override;
+};
\ No newline at end of file
diff --git a/Engine/Source/Runtime/ClothingSystemRuntimeNv/Public/ClothingSimulationInteractorNv.h b/Engine/Source/Runtime/ClothingSystemRuntimeNv/Public/ClothingSimulationInteractorNv.h
new file mode 100644
--- /dev/null
+++ b/Engine/Source/Runtime/ClothingSystemRuntimeNv/Public/ClothingSimulationInteractorNv.h
@@ -0,0 +1,48 @@
+// Copyright 1998-2019 Epic Games, Inc. All Rights Reserved.
+
+#pragma once
+
+#include "ClothingSimulationInteractor.h"
+#include "ClothingSimulationNv.h"
+#include "ClothingSimulationInteractorNv.generated.h"
+
+class FClothingSimulationNv;
+class FClothingSimulationContextNv;
+
+// Command signature for handling synced command buffer
+DECLARE_DELEGATE_TwoParams(NvInteractorCommand, FClothingSimulationNv*, FClothingSimulationContextNv*)
+
+UCLASS(BlueprintType)
+class CLOTHINGSYSTEMRUNTIMENV_API UClothingSimulationInteractorNv : public UClothingSimulationInteractor
+{
+	GENERATED_BODY()
+
+public:
+
+	// UClothingSimulationInteractor Interface
+	virtual void PhysicsAssetUpdated() override;
+	virtual void ClothConfigUpdated() override;
+	virtual void Sync(IClothingSimulation* InSimulation, IClothingSimulationContext* InContext) override;
+	//////////////////////////////////////////////////////////////////////////
+
+	// Set the stiffness of the spring force for anim drive
+	UFUNCTION(BlueprintCallable, Category=ClothingSimulation)
+	void SetAnimDriveSpringStiffness(float InStiffness);
+
+	// Set the stiffness of the resistive damping force for anim drive
+	UFUNCTION(BlueprintCallable, Category=ClothingSimulation)
+	void SetAnimDriveDamperStiffness(float InStiffness);
+
+	// Set a new gravity override and enable the override
+	UFUNCTION(BlueprintCallable, Category=ClothingSimulation)
+	void EnableGravityOverride(const FVector& InVector);
+
+	// Disable any currently set gravity override
+	UFUNCTION(BlueprintCallable, Category = ClothingSimulation)
+	void DisableGravityOverride();
+
+private:
+
+	// Command queue processed when we hit a sync
+	TArray<NvInteractorCommand> Commands;
+};
diff --git a/Engine/Source/Runtime/ClothingSystemRuntimeNv/Public/ClothingSimulationNv.h b/Engine/Source/Runtime/ClothingSystemRuntimeNv/Public/ClothingSimulationNv.h
new file mode 100644
--- /dev/null
+++ b/Engine/Source/Runtime/ClothingSystemRuntimeNv/Public/ClothingSimulationNv.h
@@ -0,0 +1,276 @@
+// Copyright 1998-2019 Epic Games, Inc. All Rights Reserved.
+
+#pragma once
+
+#if WITH_NVCLOTH
+
+#include "ClothingSimulation.h"
+#include "ClothingActor.h"
+#include "ClothingSimulationContext.h"
+
+#include "Assets/ClothingAsset.h"
+
+namespace nv
+{
+	namespace cloth
+	{
+		class Fabric;
+		class Cloth;
+		class Solver;
+		class Factory;
+		struct PhaseConfig;
+	}
+}
+
+namespace physx
+{
+	class PxVec4;
+}
+
+class FClothingSimulationContextNv final : public FClothingSimulationContextBase
+{
+public:
+
+	// Ref to local matrices from the owning component (for skinning fixed verts)
+	TArray<FMatrix> RefToLocals;
+
+};
+
+// Scratch data for simulation to avoid allocations while processing, per actor data
+struct FClothingActorScratchData
+{
+	void Reset();
+
+	TArray<physx::PxVec4> SphereData;
+	TArray<uint32> CapsuleSphereIndices;
+	TArray<physx::PxVec4> PlaneData;
+	TArray<uint32> ConvexMasks;
+	TArray<FVector> ParticleVelocities;
+};
+
+class FClothingActorNv final : public FClothingActorBase
+{
+public:
+
+	FClothingActorNv();
+
+	// Runtime data need per-lod for each actor
+	struct FActorLodData
+	{
+		// The fabric object created for this LOD
+		nv::cloth::Fabric* Fabric;
+
+		// The cloth (simulation) object for this LOD
+		nv::cloth::Cloth* Cloth;
+
+		// Original rest positions for this LOD (needed for self collisions)
+		TArray<physx::PxVec4> Px_RestPositions;
+
+		// List of phase configs for this actor. Phases are the different
+		// constraint groups that are present (horz, vert, stretch, shear, bend)
+		TArray<nv::cloth::PhaseConfig> PhaseConfigs;
+
+		// The types of each phase in the above PhaseConfigs array
+		TArray<int32> PhaseTypes;
+	};
+
+	// Skins the physics mesh to the current bone transforms, necessary to build motion constraints
+	// Stores the results internally (see below)
+	void SkinPhysicsMesh(FClothingSimulationContextNv* InContext);
+
+	// Updates the motion constraints for this actor (needs a skinned physics mesh)
+	void UpdateMotionConstraints(FClothingSimulationContextNv* InContext);
+
+	// Updates the wind effects on the currently active cloth
+	void UpdateWind(FClothingSimulationContextNv* InContext, const FVector& InWindVelocity);
+
+	// Conditional rebuild of the aggregated collisions list
+	void ConditionalRebuildCollisions();
+
+	// Updates the anim drive springs to push the simulation back to the skinned location
+	void UpdateAnimDrive(FClothingSimulationContextNv* InContext);
+
+	// Current stiffnesses for anim drive, can be overridden by the interactor
+	float CurrentAnimDriveSpringStiffness;
+	float CurrentAnimDriveDamperStiffness;
+
+	// Gravity override for this actor, can be overriden by the interactor
+	bool bUseGravityOverride;
+	FVector GravityOverride;
+
+	// Cache for previous state for handling teleports
+	FVector LastVelocity;
+	FTransform LastRootTransform;
+
+private:
+
+	// Builds a list of particle velocities given the current simulation state
+	// Required to emulate the legacy wind method that APEX used to use
+	void CalculateParticleVelocities(TArray<FVector>& OutVelocities);
+
+	// Current cloth (not mesh) LOD that this actor is running
+	int32 CurrentLodIndex;
+
+	// List of all active collisions currently affecting the simulation
+	FClothCollisionData AggregatedCollisions;
+
+	// List of collisions that were injected from an external source
+	FClothCollisionData ExternalCollisions;
+
+	// Collisions extracted from our physics asset
+	FClothCollisionData ExtractedCollisions;
+
+	// Whether or not we need to rebuild our collisions on the next simulation setp
+	bool bCollisionsDirty;
+
+	// Index to write back to on GetSimulationData for this actor
+	int32 SimDataIndex;
+
+	// Actual clothing LOD data (simulation objects)
+	TArray<FActorLodData> LodData;
+
+	// How we're going to calculate our wind data (see EClothingWindMethod for method descriptions)
+	EClothingWindMethod WindMethod;
+
+	// Thickness to add to collisions to fake cloth thickness
+	float CollisionThickness;
+
+	// Skinned physics mesh information for the active LOD. This is
+	// generated once per tick for the currently active clothing LOD
+	// Positions are double buffered to allow for velocity calculation
+	const TArray<FVector>& GetCurrentSkinnedPositions() const;
+	const TArray<FVector>& GetPreviousSkinnedPositions() const;
+	int32 CurrentSkinnedPositionIndex;
+	TArray<FVector> SkinnedPhysicsMeshPositions[2];
+	TArray<FVector> SkinnedPhysicsMeshNormals;
+
+	// Current computed normals of the simulation mesh, the normals abov
+	// are the skinned static normals of the physics mesh, this is the current
+	// set of normals for the simulation mesh
+	TArray<FVector> CurrentNormals;
+
+	// Time step of the last tick, used for velocity calculations
+	float PreviousTimestep;
+
+	// Scratch arrays for processing during simulate, grow-only to avoid repeated allocations.
+	FClothingActorScratchData Scratch;
+
+	// Simuation given access to our data
+	friend class FClothingSimulationNv;
+};
+
+// Base simulation to fill in common data for the base context
+class FClothingSimulationBase : public IClothingSimulation
+{
+public:
+	FClothingSimulationBase();
+	virtual ~FClothingSimulationBase();
+
+protected:
+	/** Fills in the base data for a clothing simulation */
+	virtual void 
+	FillContext(
+		USkeletalMeshComponent* InComponent, 
+		float InDeltaTime, 
+		IClothingSimulationContext* InOutContext) override;
+
+	/** Maximum physics time, incoming deltas will be clamped down to this value on long frames */
+	float MaxPhysicsDelta;
+};
+
+class FClothingSimulationNv final : public FClothingSimulationBase
+{
+	// Cached from the module for speed, DO NOT DELETE, only for creating cloth objects
+	nv::cloth::Factory* CachedFactory;
+
+	// Solver object for this simulation
+	nv::cloth::Solver* Solver;
+
+	// Currently valid actors (some may not be running depending on LOD)
+	TArray<FClothingActorNv> Actors;
+
+public:
+
+	FClothingSimulationNv();
+	virtual ~FClothingSimulationNv();
+
+	// IClothingSimulation Interface
+	virtual void CreateActor(USkeletalMeshComponent* InOwnerComponent, UClothingAssetBase* InAsset, int32 InSimDataIndex) override;
+	virtual IClothingSimulationContext* CreateContext() override;
+	virtual void FillContext(USkeletalMeshComponent* InComponent, float InDeltaTime, IClothingSimulationContext* InOutContext) override;
+	virtual void Initialize() override;
+	virtual void Shutdown() override;
+	virtual bool ShouldSimulate() const override;
+	virtual void Simulate(IClothingSimulationContext* InContext) override;
+	virtual FBoxSphereBounds GetBounds(const USkeletalMeshComponent* InOwnerComponent) const override;
+
+	virtual void DestroyActors() override;
+	virtual void DestroyContext(IClothingSimulationContext* InContext) override;
+	virtual void GetSimulationData(TMap<int32, FClothSimulData>& OutData, USkeletalMeshComponent* InOwnerComponent, USkinnedMeshComponent* InOverrideComponent) const override;
+	virtual void AddExternalCollisions(const FClothCollisionData& InData) override;
+	virtual void ClearExternalCollisions() override;
+	virtual void GetCollisions(FClothCollisionData& OutCollisions, bool bIncludeExternal = true) const override;
+	virtual void GatherStats() const override;
+	//////////////////////////////////////////////////////////////////////////
+
+	// Functions to be called from the game thread only when the simulation is not running
+
+	// Refresh config data if the base config changes
+	CLOTHINGSYSTEMRUNTIMENV_API void RefreshClothConfig();
+
+	// Clear and re-extract all physics bodies from our physics asset
+	CLOTHINGSYSTEMRUNTIMENV_API void RefreshPhysicsAsset();
+
+	// Given a callable object, call for each actor
+	template<typename Lambda>
+	void ExecutePerActor(Lambda InCallable)
+	{
+		for(FClothingActorNv& Actor : Actors)
+		{
+			InCallable(Actor);
+		}
+	}
+
+private:
+
+
+	// Update the LOD for the current actors, this is more complex than just updating a LOD value,
+	// we need to skin the incoming simulation mesh to the outgoing mesh (the weighting data should have been
+	// built in the asset already) to make sure it matches up without popping
+	// @param InPredictedLod The predicted LOD for the mesh component this clothing simulation is running on
+	// @param ComponentToWorld The component to world transform for the mesh component this clothing simulation is running on
+	// @param CSTranforms Component space transforms of the owning skeletal mesh component
+	// @param RefToLocals RefToLocals of the owning skeletal mesh component to init simulation mesh
+	// @param bForceNoRemap When changing LODs The incoming LOD can be skinned to the outgoing LOD for better transitions - this flag skips that step
+	// @param bForceActorChecks Typically we trust all LODs to be in sync, but that isn't always the case (e.g. adding a new actor). This forces each actor's LOD to be checked instead of trusting the master LOD
+	void UpdateLod(int32 InPredictedLod, const FTransform& ComponentToWorld, const TArray<FTransform>& CSTransforms, const TArray<FMatrix>& RefToLocals, bool bForceNoRemap = false, bool bForceActorChecks = false);
+
+	// The core simulation is only solving unoriented particles, so we need to compute normals after the
+	// simulation runs
+	void ComputePhysicalMeshNormals(FClothingActorNv &Actor);
+
+	// Given a clothing config from an asset, apply it to the provided actor. Currently
+	// this is only used from CreateActor, but could be exposed for runtime changes
+	void ApplyClothConfig(FClothConfig &Config, FClothingActorNv &InActor);
+
+	// Extract collisions from the physics asset inside Asset and apply them to InActor
+	// Not safe to call from workers (i.e. inside the simulation).
+	void ExtractActorCollisions(UClothingAsset* Asset, FClothingActorNv &InActor);
+
+	// The current LOD index for the owning skeletal mesh component
+	int32 CurrentMeshLodIndex;
+
+#if WITH_EDITOR
+public:
+	CLOTHINGSYSTEMRUNTIMENV_API void DebugDraw_PhysMesh(USkeletalMeshComponent* OwnerComponent, FPrimitiveDrawInterface* PDI) const;
+	CLOTHINGSYSTEMRUNTIMENV_API void DebugDraw_Normals(USkeletalMeshComponent* OwnerComponent, FPrimitiveDrawInterface* PDI) const;
+	CLOTHINGSYSTEMRUNTIMENV_API void DebugDraw_Collision(USkeletalMeshComponent* OwnerComponent, FPrimitiveDrawInterface* PDI) const;
+	CLOTHINGSYSTEMRUNTIMENV_API void DebugDraw_Backstops(USkeletalMeshComponent* OwnerComponent, FPrimitiveDrawInterface* PDI) const;
+	CLOTHINGSYSTEMRUNTIMENV_API void DebugDraw_MaxDistances(USkeletalMeshComponent* OwnerComponent, FPrimitiveDrawInterface* PDI) const;
+	CLOTHINGSYSTEMRUNTIMENV_API void DebugDraw_SelfCollision(USkeletalMeshComponent* OwnerComponent, FPrimitiveDrawInterface* PDI) const;
+	CLOTHINGSYSTEMRUNTIMENV_API void DebugDraw_AnimDrive(USkeletalMeshComponent* OwnerComponent, FPrimitiveDrawInterface* PDI) const;
+
+#endif
+};
+
+#endif
diff --git a/Engine/Source/Runtime/ClothingSystemRuntimeNv/Public/ClothingSystemRuntimeModuleNv.h b/Engine/Source/Runtime/ClothingSystemRuntimeNv/Public/ClothingSystemRuntimeModuleNv.h
new file mode 100644
--- /dev/null
+++ b/Engine/Source/Runtime/ClothingSystemRuntimeNv/Public/ClothingSystemRuntimeModuleNv.h
@@ -0,0 +1,58 @@
+// Copyright 1998-2019 Epic Games, Inc. All Rights Reserved.
+
+#pragma once
+
+#include "Modules/ModuleInterface.h"
+#include "ClothingSimulationFactory.h"
+
+namespace nv
+{
+	namespace cloth
+	{
+		class Factory;
+		class ClothMeshQuadifier;
+	}
+}
+
+class FClothingSystemRuntimeModuleNv : public IModuleInterface, public IClothingSimulationFactoryClassProvider
+{
+
+public:
+
+	FClothingSystemRuntimeModuleNv();
+
+	virtual void StartupModule() override;
+	virtual void ShutdownModule() override;
+
+	nv::cloth::Factory* GetSoftwareFactory();
+	nv::cloth::ClothMeshQuadifier* GetMeshQuadifier();
+
+	// IClothingSimulationFactoryClassProvider Interface
+	virtual UClass* GetDefaultSimulationFactoryClass() override;
+	//////////////////////////////////////////////////////////////////////////
+
+private:
+
+#if WITH_NVCLOTH
+	nv::cloth::Factory* ClothFactory;
+	nv::cloth::ClothMeshQuadifier* Quadifier;
+
+	void DelayLoadNvCloth();
+	void ShutdownNvClothLibs();
+
+#if PLATFORM_WINDOWS || PLATFORM_MAC
+	void* NvClothHandle;
+#endif
+
+#if PLATFORM_WINDOWS
+	void DelayLoadNvCloth_Windows();
+	void ShutdownNvCloth_Windows();
+#endif
+
+#if PLATFORM_MAC
+	void DelayLoadNvCloth_Mac();
+	void ShutdownNvCloth_Mac();
+#endif
+
+#endif
+};
diff --git a/Engine/Source/Runtime/ClothingSystemRuntimeNv/Public/NvClothIncludes.h b/Engine/Source/Runtime/ClothingSystemRuntimeNv/Public/NvClothIncludes.h
new file mode 100644
--- /dev/null
+++ b/Engine/Source/Runtime/ClothingSystemRuntimeNv/Public/NvClothIncludes.h
@@ -0,0 +1,52 @@
+// Copyright 1998-2019 Epic Games, Inc. All Rights Reserved.
+
+#pragma once
+
+#include "CoreMinimal.h"
+#include "EngineDefines.h"
+
+#ifdef NV_CLOTH_API
+#undef NV_CLOTH_API	//Our build system treats PhysX as a module and automatically defines this. PhysX has its own way of defining this.
+#endif
+
+#if WITH_NVCLOTH
+
+PUSH_MACRO(check)
+#undef check
+
+#if PLATFORM_XBOXONE
+	#pragma pack(push,16)
+#elif PLATFORM_WINDOWS
+	#if PLATFORM_64BITS
+		#pragma pack(push,16)
+	#elif PLATFORM_32BITS
+		#pragma pack(push,8)
+	#endif
+#endif
+
+#include "PxErrors.h"
+#include "PxErrorCallback.h"
+
+#include "NvCloth/PhaseConfig.h"
+#include "NvCloth/Callbacks.h"
+#include "NvCloth/Factory.h"
+#include "NvCloth/Solver.h"
+#include "NvCloth/Cloth.h"
+
+#include "NvClothExt/ClothMeshDesc.h"
+#include "NvClothExt/ClothFabricCooker.h"
+#include "NvClothExt/ClothMeshQuadifier.h"
+
+#if PLATFORM_XBOXONE
+	#pragma pack(pop)
+#elif PLATFORM_WINDOWS
+	#if PLATFORM_64BITS
+		#pragma pack(pop)
+	#elif PLATFORM_32BITS
+		#pragma pack(pop)
+	#endif
+#endif
+
+POP_MACRO(check)
+
+#endif
diff --git a/Engine/Source/Runtime/ClothingSystemRuntimeNv/Public/NvClothSupport.h b/Engine/Source/Runtime/ClothingSystemRuntimeNv/Public/NvClothSupport.h
new file mode 100644
--- /dev/null
+++ b/Engine/Source/Runtime/ClothingSystemRuntimeNv/Public/NvClothSupport.h
@@ -0,0 +1,86 @@
+// Copyright 1998-2019 Epic Games, Inc. All Rights Reserved.
+
+#pragma once
+
+#if WITH_NVCLOTH
+
+#include "NvClothIncludes.h"
+
+#include "Logging/LogVerbosity.h"
+
+#include "Containers/UnrealString.h"
+#include "Containers/Array.h"
+
+namespace NvClothSupport
+{
+	namespace Constants
+	{
+		const static float UnitConversionScale = 100.0f;
+		const static float UnitConversionScaleCube = UnitConversionScale * UnitConversionScale * UnitConversionScale;
+	}
+
+	ELogVerbosity::Type PxErrorToLogVerbosity(physx::PxErrorCode::Enum InEnum);
+	FString PxErrorToString(physx::PxErrorCode::Enum InEnum);
+
+	// Helper incase we ever want a new allocator for NvCloth
+	class physx::PxAllocatorCallback* GetAllocator();
+
+	void InitializeNvClothingSystem();
+	void ShutdownNvClothingSystem();
+
+	struct ClothTri
+	{
+		ClothTri()
+		{}
+
+		ClothTri(uint32 A, uint32 B, uint32 C)
+		{
+			T[0] = A;
+			T[1] = B;
+			T[2] = C;
+		}
+
+		uint32 T[3];
+	};
+
+	struct ClothQuad
+	{
+		ClothQuad()
+		{}
+
+		ClothQuad(uint32 A, uint32 B, uint32 C, uint32 D)
+		{
+			Q[0] = A;
+			Q[1] = B;
+			Q[2] = C;
+			Q[3] = D;
+		}
+
+		uint32 Q[4];
+	};
+
+	template<typename ContainerType>
+	nv::cloth::Range<ContainerType> CreateRange(TArray<ContainerType>& InArray, int32 InBeginOffset = 0)
+	{
+		return {InArray.GetData() + InBeginOffset, (InArray.GetData() + InBeginOffset) + InArray.Num()};
+	}
+
+	template<typename ContainerType>
+	nv::cloth::Range<const ContainerType> CreateRange(const TArray<ContainerType>& InArray, int32 InBeginOffset = 0)
+	{
+		return{InArray.GetData() + InBeginOffset, (InArray.GetData() + InBeginOffset) + InArray.Num()};
+	}
+
+	struct ClothParticleScopeLock
+	{
+		ClothParticleScopeLock();
+		ClothParticleScopeLock(nv::cloth::Cloth* InCloth);
+
+		~ClothParticleScopeLock();
+
+	private:
+		nv::cloth::Cloth* LockedCloth;
+	};
+}
+
+#endif
diff --git a/Engine/Source/Runtime/ClothingSystemRuntimeNv/Public/Utils/ClothingMeshUtils.h b/Engine/Source/Runtime/ClothingSystemRuntimeNv/Public/Utils/ClothingMeshUtils.h
new file mode 100644
--- /dev/null
+++ b/Engine/Source/Runtime/ClothingSystemRuntimeNv/Public/Utils/ClothingMeshUtils.h
@@ -0,0 +1,193 @@
+// Copyright 1998-2019 Epic Games, Inc. All Rights Reserved.
+
+#pragma once
+
+#include "SkeletalMeshTypes.h"
+
+DECLARE_LOG_CATEGORY_EXTERN(LogClothingMeshUtils, Log, All);
+
+struct FClothPhysicalMeshData;
+
+namespace ClothingMeshUtils
+{
+	struct CLOTHINGSYSTEMRUNTIMENV_API ClothMeshDesc
+	{
+		ClothMeshDesc(TArrayView<const FVector> InPositions, TArrayView<const FVector> InNormals,  TArrayView<const uint32> InIndices)
+			: Positions(InPositions)
+			, Normals(InNormals)
+			, Indices(InIndices)
+		{}
+
+		bool HasValidMesh() const
+		{
+			return Positions.Num() == Normals.Num() && Indices.Num() % 3 == 0;
+		}
+
+		TArrayView<const FVector> Positions;
+		TArrayView<const FVector> Normals;
+		TArrayView<const uint32> Indices;
+	};
+
+	// Static method for calculating a skinned mesh result from source data
+	void CLOTHINGSYSTEMRUNTIMENV_API
+	SkinPhysicsMesh(
+		const TArray<int32>& BoneMap, // UClothingAsset::UsedBoneIndices
+		const FClothPhysicalMeshData& InMesh, 
+		const FTransform& RootBoneTransform, 
+		const FMatrix* InBoneMatrices, 
+		const int32 InNumBoneMatrices, 
+		TArray<FVector>& OutPositions, 
+		TArray<FVector>& OutNormals);
+
+	/**
+	* Given mesh information for two meshes, generate a list of skinning data to embed mesh0 in mesh1
+	* @param OutSkinningData	- Final skinning data to map mesh0 to mesh1
+	* @param TargetMesh		- Mesh data for the mesh we are embedding
+	* @param TargetTangents		- Optional Tangents for the mesh we are embedding
+	* @param Mesh1Normals		- Vertex normals for Mesh1
+	* @param Mesh1Indices		- Triangle indices for Mesh1
+	*/
+	void CLOTHINGSYSTEMRUNTIMENV_API GenerateMeshToMeshSkinningData(
+		TArray<FMeshToMeshVertData>& OutSkinningData,
+		const ClothMeshDesc& TargetMesh,
+		const TArray<FVector>* TargetTangents,
+		const ClothMeshDesc& SourceMesh);
+
+	/** 
+	 * Embeds a list of positions into a source mesh
+	 * @param SourceMesh The mesh to embed in
+	 * @param Positions The positions to embed in SourceMesh
+	 * @param OutEmbeddedPositions Embedded version of the original positions, a barycentric coordinate and distance along the normal of the triangle
+	 * @param OutSourceIndices Source index list for the embedded positions, 3 per position to denote the source triangle
+	 */
+	void CLOTHINGSYSTEMRUNTIMENV_API GenerateEmbeddedPositions(
+		const ClothMeshDesc& SourceMesh, 
+		TArrayView<const FVector> Positions, 
+		TArray<FVector4>& OutEmbeddedPositions, 
+		TArray<int32>& OutSourceIndices);
+
+	/**
+	* Given a triangle ABC with normals at each vertex NA, NB and NC, get a barycentric coordinate
+	* and corresponding distance from the triangle encoded in an FVector4 where the components are
+	* (BaryX, BaryY, BaryZ, Dist)
+	* @param A		- Position of triangle vertex A
+	* @param B		- Position of triangle vertex B
+	* @param C		- Position of triangle vertex C
+	* @param NA	- Normal at vertex A
+	* @param NB	- Normal at vertex B
+	* @param NC	- Normal at vertex C
+	* @param Point	- Point to calculate Bary+Dist for
+	*/
+	FVector4 GetPointBaryAndDist(
+		const FVector& A,
+		const FVector& B,
+		const FVector& C,
+		const FVector& NA,
+		const FVector& NB,
+		const FVector& NC,
+		const FVector& Point);
+
+	/** 
+	 * Object used to map vertex parameters between two meshes using the
+	 * same barycentric mesh to mesh mapping data we use for clothing
+	 */
+	class CLOTHINGSYSTEMRUNTIMENV_API FVertexParameterMapper
+	{
+	public:
+		FVertexParameterMapper() = delete;
+		FVertexParameterMapper(const FVertexParameterMapper& Other) = delete;
+
+		FVertexParameterMapper(TArrayView<const FVector> InMesh0Positions,
+			TArrayView<const FVector> InMesh0Normals,
+			TArrayView<const FVector> InMesh1Positions,
+			TArrayView<const FVector> InMesh1Normals,
+			TArrayView<const uint32> InMesh1Indices)
+			: Mesh0Positions(InMesh0Positions)
+			, Mesh0Normals(InMesh0Normals)
+			, Mesh1Positions(InMesh1Positions)
+			, Mesh1Normals(InMesh1Normals)
+			, Mesh1Indices(InMesh1Indices)
+		{
+
+		}
+
+		/** Generic mapping function, can be used to map any type with a provided callable */
+		template<typename T, typename Lambda>
+		void Map(TArrayView<const T>& SourceData, TArray<T>& DestData, const Lambda& Func)
+		{
+			// Enforce the interp func signature (returns T and takes a bary and 3 Ts)
+			// If you hit this then either the return type isn't T or your arguments aren't convertible to T
+			static_assert(TAreTypesEqual<T, typename TDecay<decltype(Func(DeclVal<FVector>(), DeclVal<T>(), DeclVal<T>(), DeclVal<T>()))>::Type>::Value, "Invalid Lambda signature passed to Map");
+
+			const int32 NumMesh0Positions = Mesh0Positions.Num();
+			const int32 NumMesh0Normals = Mesh0Normals.Num();
+
+			const int32 NumMesh1Positions = Mesh1Positions.Num();
+			const int32 NumMesh1Normals = Mesh1Normals.Num();
+			const int32 NumMesh1Indices = Mesh1Indices.Num();
+
+			// Validate mesh data
+			check(NumMesh0Positions == NumMesh0Normals);
+			check(NumMesh1Positions == NumMesh1Normals);
+			check(NumMesh1Indices % 3 == 0);
+			check(SourceData.Num() == NumMesh1Positions);
+
+			if(DestData.Num() != NumMesh0Positions)
+			{
+				DestData.Reset();
+				DestData.AddUninitialized(NumMesh0Positions);
+			}
+
+			ClothMeshDesc SourceMeshDesc(Mesh1Positions, Mesh1Normals, Mesh1Indices);
+
+			TArray<FVector4> EmbeddedPositions;
+			TArray<int32> SourceIndices;
+			ClothingMeshUtils::GenerateEmbeddedPositions(SourceMeshDesc, Mesh0Positions, EmbeddedPositions, SourceIndices);
+
+			for(int32 DestVertIndex = 0 ; DestVertIndex < NumMesh0Positions ; ++DestVertIndex)
+			{
+				// Truncate the distance from the position data
+				FVector Bary = EmbeddedPositions[DestVertIndex];
+
+				const int32 SourceTriBaseIdx = DestVertIndex * 3;
+				T A = SourceData[SourceIndices[SourceTriBaseIdx + 0]];
+				T B = SourceData[SourceIndices[SourceTriBaseIdx + 1]];
+				T C = SourceData[SourceIndices[SourceTriBaseIdx + 2]];
+
+				T& DestVal = DestData[DestVertIndex];
+
+				// If we're super close to a vertex just take it's value.
+				// Otherwise call the provided interp lambda
+				const static FVector OneVec(1.0f, 1.0f, 1.0f);
+				FVector DiffVec = OneVec - Bary;
+				if(FMath::Abs(DiffVec.X) <= SMALL_NUMBER)
+				{
+					DestVal = A;
+				}
+				else if(FMath::Abs(DiffVec.Y) <= SMALL_NUMBER)
+				{
+					DestVal = B;
+				}
+				else if(FMath::Abs(DiffVec.Z) <= SMALL_NUMBER)
+				{
+					DestVal = C;
+				}
+				else
+				{
+					DestData[DestVertIndex] = Func(Bary, A, B, C);
+				}
+			}
+		}
+
+		// Defined type mappings for brevity
+		void Map(TArrayView<const float> Source, TArray<float>& Dest);
+
+	private:
+
+		TArrayView<const FVector> Mesh0Positions;
+		TArrayView<const FVector> Mesh0Normals;
+		TArrayView<const FVector> Mesh1Positions;
+		TArrayView<const FVector> Mesh1Normals;
+		TArrayView<const uint32> Mesh1Indices;
+	};
+}
\ No newline at end of file
diff --git a/Engine/Source/Runtime/Core/Public/UObject/PhysicsObjectVersion.h b/Engine/Source/Runtime/Core/Public/UObject/PhysicsObjectVersion.h
--- a/Engine/Source/Runtime/Core/Public/UObject/PhysicsObjectVersion.h
+++ b/Engine/Source/Runtime/Core/Public/UObject/PhysicsObjectVersion.h
@@ -12,6 +12,9 @@ struct CORE_API FPhysicsObjectVersion
 		// Before any version changes were made
 		BeforeCustomVersionWasAdded = 0,
 
+		// Temporary versioning to serialize hair guide data in addition to hair render data
+		HairAssetSerialization,
+		HairAssetSerialization_V2,
 
 		// -----<new versions can be added above this line>-------------------------------------------------
 		VersionPlusOne,
diff --git a/Engine/Source/Runtime/D3D12RHI/Private/D3D12RayTracing.cpp b/Engine/Source/Runtime/D3D12RHI/Private/D3D12RayTracing.cpp
--- a/Engine/Source/Runtime/D3D12RHI/Private/D3D12RayTracing.cpp
+++ b/Engine/Source/Runtime/D3D12RHI/Private/D3D12RayTracing.cpp
@@ -1851,11 +1851,15 @@ FRayTracingGeometryRHIRef FD3D12DynamicRHI::RHICreateRayTracingGeometry(const FR
 	if (Initializer.GeometryType == RTGT_Triangles)
 	{
 		// #dxr_todo UE-72160: VET_Half4 (DXGI_FORMAT_R16G16B16A16_FLOAT) is also supported by DXR. Should we support it?
-		check(Initializer.VertexBufferElementType == VET_Float3 || Initializer.VertexBufferElementType == VET_Float2 || Initializer.VertexBufferElementType == VET_Half2);
+		check(Initializer.VertexBufferElementType == VET_Float4 || Initializer.VertexBufferElementType == VET_Float3 || Initializer.VertexBufferElementType == VET_Float2 || Initializer.VertexBufferElementType == VET_Half2);
 
 		// #dxr_todo UE-72160: temporary constraints on vertex and index buffer formats (this will be relaxed when more flexible vertex/index fetching is implemented)
-		checkf(Initializer.VertexBufferElementType == VET_Float3, TEXT("Only float3 vertex buffers are currently implemented.")); // #dxr_todo UE-72160: support other vertex buffer formats
-		checkf(Initializer.VertexBufferStride == 12, TEXT("Only deinterleaved float3 position vertex buffers are currently implemented.")); // #dxr_todo UE-72160: support interleaved vertex buffers
+		checkf(Initializer.VertexBufferElementType == VET_Float3 || Initializer.VertexBufferElementType == VET_Float4, TEXT("Only float3 vertex buffers are currently implemented.")); // #dxr_todo UE-72160: support other vertex buffer formats
+		if (Initializer.VertexBufferElementType == VET_Float3)
+			checkf(Initializer.VertexBufferStride >= 12, TEXT("Only deinterleaved float3 position vertex buffers are currently implemented.")); // #dxr_todo UE-72160: support interleaved vertex buffers
+
+		if (Initializer.VertexBufferElementType == VET_Float4)
+			checkf(Initializer.VertexBufferStride >= 16, TEXT("Only deinterleaved float3 position vertex buffers are currently implemented.")); // #dxr_todo UE-72160: support interleaved vertex buffers
 	}
 
 	if (Initializer.GeometryType == RTGT_Procedural)
@@ -2098,6 +2102,10 @@ void FD3D12RayTracingGeometry::BuildAccelerationStructure(FD3D12CommandContext&
 		{
 			switch (VertexElemType)
 			{
+			case VET_Float4: 
+				// While the DXGI_FORMAT_R32G32B32A32_FLOAT format is not supported by DXR, since we manually load vertex 
+				// data when we are building the BLAS, we can just rely on the vertex stride to offset the read index, 
+				// and read only the 3 vertex components, and so use the DXGI_FORMAT_R32G32B32_FLOAT vertex format
 			case VET_Float3:
 				Desc.Triangles.VertexFormat = DXGI_FORMAT_R32G32B32_FLOAT;
 				break;
@@ -2547,7 +2555,7 @@ FD3D12RayTracingShaderTable* FD3D12RayTracingScene::FindOrCreateShaderTable(cons
 				if (Geometry->GeometryType == D3D12_RAYTRACING_GEOMETRY_TYPE_TRIANGLES)
 				{
 					// #dxr_todo UE-72160: support various vertex buffer layouts (fetch/decode based on vertex stride and format)
-					checkf(Geometry->VertexElemType == VET_Float3, TEXT("Only VET_Float3 is currently implemented and tested. Other formats will be supported in the future."));
+					checkf(Geometry->VertexElemType == VET_Float3 || Geometry->VertexElemType == VET_Float4, TEXT("Only VET_Float3 is currently implemented and tested. Other formats will be supported in the future."));
 				}
 
 				SystemParameters.RootConstants.SetVertexAndIndexStride(Geometry->VertexStrideInBytes, IndexStride);
diff --git a/Engine/Source/Runtime/Engine/Classes/Components/LightComponentBase.h b/Engine/Source/Runtime/Engine/Classes/Components/LightComponentBase.h
--- a/Engine/Source/Runtime/Engine/Classes/Components/LightComponentBase.h
+++ b/Engine/Source/Runtime/Engine/Classes/Components/LightComponentBase.h
@@ -77,6 +77,12 @@ class ENGINE_API ULightComponentBase : public USceneComponent
 	UPROPERTY(EditAnywhere, BlueprintReadOnly, Category = Light, AdvancedDisplay)
 	uint32 bCastVolumetricShadow : 1;
 
+	/**
+	 * Whether the light should cast high quality hair-strands self-shadowing. When this option is enabled, an extra GPU cost for this light. 
+	 **/
+	UPROPERTY(EditAnywhere, BlueprintReadOnly, Category = Light, AdvancedDisplay)
+	uint32 bCastDeepShadow : 1;
+
 	/** Whether the light shadows are computed with shadow-mapping or ray-tracing (when available). */
 	UPROPERTY(EditAnywhere, BlueprintReadOnly, Category = Light, meta = (DisplayName = "Cast Ray Tracing Shadow"), AdvancedDisplay)
 	uint32 bCastRaytracedShadow : 1;
@@ -133,6 +139,9 @@ class ENGINE_API ULightComponentBase : public USceneComponent
 	UFUNCTION(BlueprintCallable, Category = "Rendering|Components|Light")
 	void SetCastVolumetricShadow(bool bNewValue);
 
+	UFUNCTION(BlueprintCallable, Category = "Rendering|Components|Light")
+	void SetCastDeepShadow(bool bNewValue);
+
 	UFUNCTION(BlueprintCallable, Category = "Rendering|Components|Light")
 	void SetAffectReflection(bool bNewValue);
 
diff --git a/Engine/Source/Runtime/Engine/Classes/Engine/Engine.h b/Engine/Source/Runtime/Engine/Classes/Engine/Engine.h
--- a/Engine/Source/Runtime/Engine/Classes/Engine/Engine.h
+++ b/Engine/Source/Runtime/Engine/Classes/Engine/Engine.h
@@ -857,11 +857,27 @@ public:
 	/** The material used to render wireframe meshes. */
 	UPROPERTY()
 	class UMaterial* WireframeMaterial;
-
+	
 	/** @todo document */
 	UPROPERTY(globalconfig)
 	FString WireframeMaterialName;
 
+	/** The material used to default hair meshes. */
+	UPROPERTY()
+	class UMaterial* HairDefaultMaterial;
+
+	/** @todo document */
+	UPROPERTY(globalconfig)
+	FString HairDefaultMaterialName;
+
+	/** The material used to debug hair meshes. */
+	UPROPERTY()
+	class UMaterial* HairDebugMaterial;
+
+	/** @todo document */
+	UPROPERTY(globalconfig)
+	FString HairDebugMaterialName;
+
 #if WITH_EDITORONLY_DATA
 	/** A translucent material used to render things in geometry mode. */
 	UPROPERTY()
diff --git a/Engine/Source/Runtime/Engine/Classes/Materials/Material.h b/Engine/Source/Runtime/Engine/Classes/Materials/Material.h
--- a/Engine/Source/Runtime/Engine/Classes/Materials/Material.h
+++ b/Engine/Source/Runtime/Engine/Classes/Materials/Material.h
@@ -734,6 +734,13 @@ public:
 	UPROPERTY(EditAnywhere, BlueprintReadOnly, Category = Usage)
 	uint32 bUsedWithWater : 1;
 
+	/**
+	 * Indicates that the material and its instances can be use with hair strands
+	 * This will result in the shaders required to support hair strands geometries being compiled which will increase shader compile time and memory usage.
+	 */
+	UPROPERTY(EditAnywhere, BlueprintReadOnly, Category = Usage)
+	uint32 bUsedWithHairStrands : 1;
+
 	/** 
 	 * Indicates that the material and its instances can be used with Slate UI and UMG
 	 * This will result in the shaders required to support UI materials being compiled which will increase shader compile time and memory usage.
diff --git a/Engine/Source/Runtime/Engine/Classes/Materials/MaterialInterface.h b/Engine/Source/Runtime/Engine/Classes/Materials/MaterialInterface.h
--- a/Engine/Source/Runtime/Engine/Classes/Materials/MaterialInterface.h
+++ b/Engine/Source/Runtime/Engine/Classes/Materials/MaterialInterface.h
@@ -53,6 +53,7 @@ enum EMaterialUsage
 	MATUSAGE_NiagaraMeshParticles,
 	MATUSAGE_GeometryCache,
 	MATUSAGE_Water,
+	MATUSAGE_HairStrands,
 
 	MATUSAGE_MAX,
 };
@@ -66,6 +67,7 @@ struct ENGINE_API FMaterialRelevance
 	uint8 bOpaque : 1;
 	uint8 bMasked : 1;
 	uint8 bDistortion : 1;
+	uint8 bHairStrands : 1;
 	uint8 bSeparateTranslucency : 1; // Translucency After DOF
 	uint8 bUnderWaterTranslucency : 1;
 	uint8 bNormalTranslucency : 1;
diff --git a/Engine/Source/Runtime/Engine/Private/Collision/CollisionConversions.cpp b/Engine/Source/Runtime/Engine/Private/Collision/CollisionConversions.cpp
--- a/Engine/Source/Runtime/Engine/Private/Collision/CollisionConversions.cpp
+++ b/Engine/Source/Runtime/Engine/Private/Collision/CollisionConversions.cpp
@@ -102,6 +102,13 @@ static FVector FindGeomOpposingNormal(ECollisionShapeType QueryGeomType, const F
 		const FPhysicsShape* Shape = GetShape(Hit);
 		if (Shape)
 		{
+#if WITH_CHAOS
+			const FTransform ActorTM(Hit.Actor->R(), Hit.Actor->X());
+			const FVector LocalInNormal = ActorTM.InverseTransformVectorNoScale(InNormal);
+			const FVector LocalTraceDirectionDenorm = ActorTM.InverseTransformVectorNoScale(TraceDirectionDenorm);
+			const FVector LocalNormal = Shape->Geometry->FindGeometryOpposingNormal(LocalTraceDirectionDenorm, Hit.FaceIndex, LocalInNormal);
+			return ActorTM.TransformVectorNoScale(LocalNormal);
+#else
 			ECollisionShapeType GeomType = GetGeometryType(*Shape);
 			switch (GeomType)
 			{
@@ -111,8 +118,10 @@ static FVector FindGeomOpposingNormal(ECollisionShapeType QueryGeomType, const F
 			case ECollisionShapeType::Convex:		return FindConvexMeshOpposingNormal(Hit, TraceDirectionDenorm, InNormal);
 			case ECollisionShapeType::Heightfield:		return FindHeightFieldOpposingNormal(Hit, TraceDirectionDenorm, InNormal);
 			case ECollisionShapeType::Trimesh:	return FindTriMeshOpposingNormal(Hit, TraceDirectionDenorm, InNormal);
-			default: break; //TODO_SQ_IMPLEMENTATION check(false);	//unsupported geom type
+			default: break;
 			}
+
+#endif
 		}
 	}
 
@@ -339,13 +348,10 @@ EConvertQueryResult ConvertTraceResults(bool& OutHasValidBlockingHit, const UWor
 		{
 			if (TIsSame<HitType, FHitSweep>::Value)
 			{
-#if WITH_PHYSX
-				SetInternalFaceIndex(Hit, FindFaceIndex(Hit, Dir));
-#else
-				//#TODO Chaos implementation
-				ensure(false);
-				SetInternalFaceIndex(Hit, 0);
-#endif
+				if (!HadInitialOverlap(Hit))
+				{
+					SetInternalFaceIndex(Hit, FindFaceIndex(Hit, Dir));
+				}
 			}
 
 			FHitResult& NewResult = OutHits[OutHits.AddDefaulted()];
diff --git a/Engine/Source/Runtime/Engine/Private/Collision/SceneQueryLowLevel.cpp b/Engine/Source/Runtime/Engine/Private/Collision/SceneQueryLowLevel.cpp
--- a/Engine/Source/Runtime/Engine/Private/Collision/SceneQueryLowLevel.cpp
+++ b/Engine/Source/Runtime/Engine/Private/Collision/SceneQueryLowLevel.cpp
@@ -46,7 +46,7 @@ void FinalizeCapture(FPhysTestSerializer& Serializer)
 	if (ReplaySQs)
 	{
 		const bool bReplaySuccess = SQComparisonHelper(Serializer);
-		if (!bReplaySuccess)'ve just'
+		if (!bReplaySuccess)
 		{
 			UE_LOG(LogPhysicsCore, Warning, TEXT("Chaos SQ does not match physx"));
 			if (SerializeBadSQs && !SerializeSQs)
@@ -67,11 +67,12 @@ void FinalizeCapture(FPhysTestSerializer& Serializer) {}
 void LowLevelRaycast(FPhysScene& Scene, const FVector& Start, const FVector& Dir, float DeltaMag, FPhysicsHitCallback<FHitRaycast>& HitBuffer, EHitFlags OutputFlags, FQueryFlags QueryFlags, const FCollisionFilterData& Filter, const FQueryFilterData& QueryFilterData, ICollisionQueryFilterCallbackBase* QueryCallback)
 {
 #if !defined(PHYSICS_INTERFACE_PHYSX) || !PHYSICS_INTERFACE_PHYSX
-	Chaos::FPhysicsSolver* Solver = Scene.GetSolver();
-	Chaos::FPhysicsSolver::FPBDRigidsEvolution* Evolution = Solver->GetEvolution();
-	FChaosSQAccelerator SQAccelerator(*Evolution);
-	//ISQAccelerator* SQAccelerator = Scene.GetSQAccelerator();
-	SQAccelerator.Raycast(Start, Dir, DeltaMag, HitBuffer, OutputFlags, QueryFilterData, *QueryCallback);
+	if (const auto& SolverAccelerationStructure = Scene.GetScene().SolverAccelerationStructure)
+	{
+		FChaosSQAccelerator SQAccelerator(*SolverAccelerationStructure.Get());
+		//ISQAccelerator* SQAccelerator = Scene.GetSQAccelerator();
+		SQAccelerator.Raycast(Start, Dir, DeltaMag, HitBuffer, OutputFlags, QueryFilterData, *QueryCallback);
+	}
 #else
 	if (SerializeSQs | ReplaySQs)
 	{
@@ -94,11 +95,12 @@ void LowLevelRaycast(FPhysScene& Scene, const FVector& Start, const FVector& Dir
 void LowLevelSweep(FPhysScene& Scene, const FPhysicsGeometry& QueryGeom, const FTransform& StartTM, const FVector& Dir, float DeltaMag, FPhysicsHitCallback<FHitSweep>& HitBuffer, EHitFlags OutputFlags, FQueryFlags QueryFlags, const FCollisionFilterData& Filter, const FQueryFilterData& QueryFilterData, ICollisionQueryFilterCallbackBase* QueryCallback)
 {
 #if !defined(PHYSICS_INTERFACE_PHYSX) || !PHYSICS_INTERFACE_PHYSX
-	Chaos::FPhysicsSolver* Solver = Scene.GetSolver();
-	Chaos::FPhysicsSolver::FPBDRigidsEvolution* Evolution = Solver->GetEvolution();
-	FChaosSQAccelerator SQAccelerator(*Evolution);
-	//ISQAccelerator* SQAccelerator = Scene.GetSQAccelerator();
-	SQAccelerator.Sweep(QueryGeom, StartTM, Dir, DeltaMag, HitBuffer, OutputFlags, QueryFilterData, *QueryCallback);
+	if (const auto& SolverAccelerationStructure = Scene.GetScene().SolverAccelerationStructure)
+	{
+		FChaosSQAccelerator SQAccelerator(*SolverAccelerationStructure.Get());
+		//ISQAccelerator* SQAccelerator = Scene.GetSQAccelerator();
+		SQAccelerator.Sweep(QueryGeom, StartTM, Dir, DeltaMag, HitBuffer, OutputFlags, QueryFilterData, *QueryCallback);
+	}
 #else
 	if (SerializeSQs | ReplaySQs)
 	{
@@ -121,11 +123,12 @@ void LowLevelSweep(FPhysScene& Scene, const FPhysicsGeometry& QueryGeom, const F
 void LowLevelOverlap(FPhysScene& Scene, const FPhysicsGeometry& QueryGeom, const FTransform& GeomPose, FPhysicsHitCallback<FHitOverlap>& HitBuffer, FQueryFlags QueryFlags, const FCollisionFilterData& Filter, const FQueryFilterData& QueryFilterData, ICollisionQueryFilterCallbackBase* QueryCallback)
 {
 #if !defined(PHYSICS_INTERFACE_PHYSX) || !PHYSICS_INTERFACE_PHYSX
-	Chaos::FPhysicsSolver* Solver = Scene.GetSolver();
-	Chaos::FPhysicsSolver::FPBDRigidsEvolution* Evolution = Solver->GetEvolution();
-	FChaosSQAccelerator SQAccelerator(*Evolution);
-	//ISQAccelerator* SQAccelerator = Scene.GetSQAccelerator();
-	SQAccelerator.Overlap(QueryGeom, GeomPose, HitBuffer, QueryFilterData, *QueryCallback);
+	if (const auto& SolverAccelerationStructure = Scene.GetScene().SolverAccelerationStructure)
+	{
+		FChaosSQAccelerator SQAccelerator(*SolverAccelerationStructure.Get());
+		//ISQAccelerator* SQAccelerator = Scene.GetSQAccelerator();
+		SQAccelerator.Overlap(QueryGeom, GeomPose, HitBuffer, QueryFilterData, *QueryCallback);
+	}
 #else
 	if (SerializeSQs | ReplaySQs)
 	{
diff --git a/Engine/Source/Runtime/Engine/Private/Components/LightComponent.cpp b/Engine/Source/Runtime/Engine/Private/Components/LightComponent.cpp
--- a/Engine/Source/Runtime/Engine/Private/Components/LightComponent.cpp
+++ b/Engine/Source/Runtime/Engine/Private/Components/LightComponent.cpp
@@ -71,6 +71,16 @@ void ULightComponentBase::SetCastVolumetricShadow(bool bNewValue)
 	}
 }
 
+void ULightComponentBase::SetCastDeepShadow(bool bNewValue)
+{
+	if (AreDynamicDataChangesAllowed()
+		&& bCastDeepShadow != bNewValue)
+	{
+		bCastDeepShadow = bNewValue;
+		MarkRenderStateDirty();
+	}
+}
+
 void ULightComponentBase::SetAffectReflection(bool bNewValue)
 {
 	if (AreDynamicDataChangesAllowed()
@@ -275,6 +285,7 @@ FLightSceneProxy::FLightSceneProxy(const ULightComponent* InLightComponent)
 	, bCastTranslucentShadows(InLightComponent->CastTranslucentShadows)
 	, bTransmission(InLightComponent->bTransmission && bCastDynamicShadow && !bStaticShadowing)
 	, bCastVolumetricShadow(InLightComponent->bCastVolumetricShadow)
+	, bCastHairStrandsDeepShadow(InLightComponent->bCastDeepShadow)
 	, bCastShadowsFromCinematicObjectsOnly(InLightComponent->bCastShadowsFromCinematicObjectsOnly)
 	, bForceCachedShadowsForMovablePrimitives(InLightComponent->bForceCachedShadowsForMovablePrimitives)
 	, bCastRaytracedShadow(InLightComponent->bCastRaytracedShadow)
@@ -673,6 +684,7 @@ void ULightComponent::PostEditChangeProperty(FPropertyChangedEvent& PropertyChan
 		PropertyName != GET_MEMBER_NAME_STRING_CHECKED(ULightComponent, LightingChannels) &&
 		PropertyName != GET_MEMBER_NAME_STRING_CHECKED(ULightComponent, VolumetricScatteringIntensity) &&
 		PropertyName != GET_MEMBER_NAME_STRING_CHECKED(ULightComponent, bCastVolumetricShadow) &&
+		PropertyName != GET_MEMBER_NAME_STRING_CHECKED(ULightComponent, bCastDeepShadow) &&
 		PropertyName != GET_MEMBER_NAME_STRING_CHECKED(ULightComponent, bCastRaytracedShadow) &&
 		PropertyName != GET_MEMBER_NAME_STRING_CHECKED(ULightComponent, bAffectReflection) &&
 		PropertyName != GET_MEMBER_NAME_STRING_CHECKED(ULightComponent, bAffectGlobalIllumination) &&
diff --git a/Engine/Source/Runtime/Engine/Private/Materials/HLSLMaterialTranslator.h b/Engine/Source/Runtime/Engine/Private/Materials/HLSLMaterialTranslator.h
--- a/Engine/Source/Runtime/Engine/Private/Materials/HLSLMaterialTranslator.h
+++ b/Engine/Source/Runtime/Engine/Private/Materials/HLSLMaterialTranslator.h
@@ -6597,6 +6597,31 @@ protected:
 		return AddCodeChunk( MCT_Float3, TEXT("MaterialExpressionBlackBody(%s)"), *GetParameterCode(Temp) );
 	}
 
+	virtual int32 GetHairUV() override
+	{
+		return AddCodeChunk(MCT_Float2, TEXT("MaterialExpressionGetHairUV(Parameters)"));
+	}
+
+	virtual int32 GetHairDimensions() override
+	{
+		return AddCodeChunk(MCT_Float2, TEXT("MaterialExpressionGetHairDimensions(Parameters)"));
+	}
+
+	virtual int32 GetHairSeed() override
+	{
+		return AddCodeChunk(MCT_Float1, TEXT("MaterialExpressionGetHairSeed(Parameters)"));
+	}
+
+	virtual int32 GetHairTangent() override
+	{
+		return AddCodeChunk(MCT_Float3, TEXT("MaterialExpressionGetHairTangent(Parameters)"));
+	}
+
+	virtual int32 GetHairRootUV() override
+	{
+		return AddCodeChunk(MCT_Float2, TEXT("MaterialExpressionGetHairRootUV(Parameters)"));
+	}
+
 	virtual int32 DistanceToNearestSurface(int32 PositionArg) override
 	{
 		if (ErrorUnlessFeatureLevelSupported(ERHIFeatureLevel::SM5) == INDEX_NONE)
diff --git a/Engine/Source/Runtime/Engine/Private/Materials/Material.cpp b/Engine/Source/Runtime/Engine/Private/Materials/Material.cpp
--- a/Engine/Source/Runtime/Engine/Private/Materials/Material.cpp
+++ b/Engine/Source/Runtime/Engine/Private/Materials/Material.cpp
@@ -1301,6 +1301,7 @@ bool UMaterial::GetUsageByFlag(EMaterialUsage Usage) const
 		case MATUSAGE_Clothing: UsageValue = bUsedWithClothing; break;
 		case MATUSAGE_GeometryCache: UsageValue = bUsedWithGeometryCache; break;
 		case MATUSAGE_Water: UsageValue = bUsedWithWater; break;
+		case MATUSAGE_HairStrands: UsageValue = bUsedWithHairStrands; break;
 		default: UE_LOG(LogMaterial, Fatal,TEXT("Unknown material usage: %u"), (int32)Usage);
 	};
 	return UsageValue;
@@ -1440,6 +1441,10 @@ void UMaterial::SetUsageByFlag(EMaterialUsage Usage, bool NewValue)
 		{
 			bUsedWithWater = NewValue; break;
 		}
+		case MATUSAGE_HairStrands:
+		{
+			bUsedWithHairStrands = NewValue; break;
+		}
 		default: UE_LOG(LogMaterial, Fatal,TEXT("Unknown material usage: %u"), (int32)Usage);
 	};
 #if WITH_EDITOR
@@ -1468,6 +1473,7 @@ FString UMaterial::GetUsageName(EMaterialUsage Usage) const
 		case MATUSAGE_Clothing: UsageName = TEXT("bUsedWithClothing"); break;
 		case MATUSAGE_GeometryCache: UsageName = TEXT("bUsedWithGeometryCache"); break;
 		case MATUSAGE_Water: UsageName = TEXT("bUsedWithWater"); break;
+		case MATUSAGE_HairStrands: UsageName = TEXT("bUsedWithHairStrands"); break;
 		default: UE_LOG(LogMaterial, Fatal,TEXT("Unknown material usage: %u"), (int32)Usage);
 	};
 	return UsageName;
diff --git a/Engine/Source/Runtime/Engine/Private/Materials/MaterialExpressions.cpp b/Engine/Source/Runtime/Engine/Private/Materials/MaterialExpressions.cpp
--- a/Engine/Source/Runtime/Engine/Private/Materials/MaterialExpressions.cpp
+++ b/Engine/Source/Runtime/Engine/Private/Materials/MaterialExpressions.cpp
@@ -102,6 +102,7 @@
 #include "Materials/MaterialExpressionGIReplace.h"
 #include "Materials/MaterialExpressionRayTracingQualitySwitch.h"
 #include "Materials/MaterialExpressionGetMaterialAttributes.h"
+#include "Materials/MaterialExpressionHairAttributes.h"
 #include "Materials/MaterialExpressionIf.h"
 #include "Materials/MaterialExpressionLightmapUVs.h"
 #include "Materials/MaterialExpressionPrecomputedAOMask.h"
@@ -16799,6 +16800,77 @@ void UMaterialExpressionCurveAtlasRowParameter::PostEditChangeProperty(FProperty
 }
 #endif
 
+
+//
+// Hair attributes
+//
+
+UMaterialExpressionHairAttributes::UMaterialExpressionHairAttributes(const FObjectInitializer& ObjectInitializer)
+	: Super(ObjectInitializer)
+{
+#if WITH_EDITORONLY_DATA
+	// Structure to hold one-time initialization
+	struct FConstructorStatics
+	{
+		FText NAME_Utility;
+		FConstructorStatics()
+			: NAME_Utility(LOCTEXT("Hair Attributes", "Hair Attributes"))
+		{
+		}
+	};
+	static FConstructorStatics ConstructorStatics;
+	MenuCategories.Add(ConstructorStatics.NAME_Utility);
+
+#endif
+
+#if WITH_EDITORONLY_DATA
+	bShowOutputNameOnPin = true;
+
+	Outputs.Reset();
+	Outputs.Add(FExpressionOutput(TEXT("U"), 1, 1, 0, 0, 0));
+	Outputs.Add(FExpressionOutput(TEXT("V"), 1, 0, 1, 0, 0));
+	Outputs.Add(FExpressionOutput(TEXT("Length"), 1, 1, 0, 0, 0));
+	Outputs.Add(FExpressionOutput(TEXT("Radius"), 1, 0, 1, 0, 0));
+	Outputs.Add(FExpressionOutput(TEXT("Seed")));
+	Outputs.Add(FExpressionOutput(TEXT("World Tangent"), 1, 1, 1, 1, 0));
+	Outputs.Add(FExpressionOutput(TEXT("Root UV"), 1, 1, 1, 0, 0));
+#endif
+}
+
+#if WITH_EDITOR
+int32 UMaterialExpressionHairAttributes::Compile(class FMaterialCompiler* Compiler, int32 OutputIndex)
+{
+	if (OutputIndex == 0 || OutputIndex == 1)
+	{
+		return Compiler->GetHairUV();
+	}
+	else if (OutputIndex == 2 || OutputIndex == 3)
+	{
+		return Compiler->GetHairDimensions();
+	}
+	else if (OutputIndex == 4)
+	{
+		return Compiler->GetHairSeed();
+	}
+	else if (OutputIndex == 5)
+	{
+		return Compiler->GetHairTangent();
+	}
+	else if (OutputIndex == 6)
+	{
+		return Compiler->GetHairRootUV();
+	}
+
+	return Compiler->Errorf(TEXT("Invalid input parameter"));
+}
+
+void UMaterialExpressionHairAttributes::GetCaption(TArray<FString>& OutCaptions) const
+{
+	OutCaptions.Add(TEXT("Hair Attributes"));
+}
+#endif // WITH_EDITOR
+
+
 //
 //  UMaterialExpressionARPassthroughCameraUVs
 //
diff --git a/Engine/Source/Runtime/Engine/Private/Materials/MaterialInterface.cpp b/Engine/Source/Runtime/Engine/Private/Materials/MaterialInterface.cpp
--- a/Engine/Source/Runtime/Engine/Private/Materials/MaterialInterface.cpp
+++ b/Engine/Source/Runtime/Engine/Private/Materials/MaterialInterface.cpp
@@ -20,6 +20,7 @@
 #include "Interfaces/ITargetPlatform.h"
 #include "Components.h"
 #include "ContentStreaming.h"
+#include "MeshBatch.h"
 
 /**
  * This is used to deprecate data that has been built with older versions.
@@ -41,6 +42,7 @@ void FMaterialRelevance::SetPrimitiveViewRelevance(FPrimitiveViewRelevance& OutV
 	OutViewRelevance.bMaskedRelevance = bMasked;
 	OutViewRelevance.bOutputsTranslucentVelocityRelevance = bOutputsTranslucentVelocity;
 	OutViewRelevance.bDistortionRelevance = bDistortion;
+	OutViewRelevance.bHairStrandsRelevance = bHairStrands;
 	OutViewRelevance.bSeparateTranslucencyRelevance = bSeparateTranslucency;
 	OutViewRelevance.bUnderWaterTranslucencyRelevance = bUnderWaterTranslucency;
 	OutViewRelevance.bNormalTranslucencyRelevance = bNormalTranslucency;
@@ -154,6 +156,7 @@ FMaterialRelevance UMaterialInterface::GetRelevance_Internal(const UMaterial* Ma
 			MaterialRelevance.bOpaque = !bIsTranslucent;
 			MaterialRelevance.bMasked = IsMasked();
 			MaterialRelevance.bDistortion = MaterialResource->IsDistorted();
+			MaterialRelevance.bHairStrands = IsCompatibleWithHairStrands(MaterialResource, InFeatureLevel);
 			MaterialRelevance.bSeparateTranslucency = bIsTranslucent && bMaterialSeparateTranclucency;
 			MaterialRelevance.bUnderWaterTranslucency = bIsTranslucent && bMaterialUnderWaterTranslucency;
 			MaterialRelevance.bNormalTranslucency = bIsTranslucent && !bMaterialSeparateTranclucency;
diff --git a/Engine/Source/Runtime/Engine/Private/Materials/MaterialShared.cpp b/Engine/Source/Runtime/Engine/Private/Materials/MaterialShared.cpp
--- a/Engine/Source/Runtime/Engine/Private/Materials/MaterialShared.cpp
+++ b/Engine/Source/Runtime/Engine/Private/Materials/MaterialShared.cpp
@@ -1131,6 +1131,11 @@ bool FMaterialResource::IsUsedWithWater() const
 	return Material->bUsedWithWater;
 }
 
+bool FMaterialResource::IsUsedWithHairStrands() const
+{
+	return Material->bUsedWithHairStrands;
+}
+
 bool FMaterialResource::IsUsedWithLandscape() const
 {
 	return false;
diff --git a/Engine/Source/Runtime/Engine/Private/PhysicsEngine/Experimental/PhysInterface_Chaos.cpp b/Engine/Source/Runtime/Engine/Private/PhysicsEngine/Experimental/PhysInterface_Chaos.cpp
--- a/Engine/Source/Runtime/Engine/Private/PhysicsEngine/Experimental/PhysInterface_Chaos.cpp
+++ b/Engine/Source/Runtime/Engine/Private/PhysicsEngine/Experimental/PhysInterface_Chaos.cpp
@@ -6,7 +6,6 @@
 #include "Physics/Experimental/PhysScene_Chaos.h"
 #include "Physics/PhysicsInterfaceTypes.h"
 #include "PhysicsEngine/AggregateGeom.h"
-#include "Chaos/Framework/SingleParticlePhysicsProxy.h"
 
 #include "Chaos/Box.h"
 #include "Chaos/Cylinder.h"
@@ -16,6 +15,7 @@
 #include "Chaos/Levelset.h"
 #include "Chaos/PBDRigidParticles.h"
 #include "Chaos/Sphere.h"
+#include "Chaos/Matrix.h"
 #include "PhysicsSolver.h"
 #include "Templates/UniquePtr.h"
 #include "ChaosSolversModule.h"
@@ -24,6 +24,8 @@
 #include "Chaos/TriangleMeshImplicitObject.h"
 #include "Chaos/Convex.h"
 #include "Chaos/GeometryQueries.h"
+#include "Chaos/Plane.h"
+
 
 
 #include "Async/ParallelFor.h"
@@ -122,6 +124,7 @@ void FPhysInterface_Chaos::ReleaseActor(FPhysicsActorHandle& Handle, FPhysScene*
 	check(Handle);
 	if (InScene)
 	{
+		InScene->GetScene().RemoveActorFromAccelerationStructure(Handle);
 		RemoveActorFromSolver(Handle, InScene->GetSolver(), FChaosSolversModule::GetModule()->GetDispatcher());
 	}
 
@@ -234,8 +237,11 @@ bool FPhysInterface_Chaos::CanSimulate_AssumesLocked(const FPhysicsActorHandle&
 
 float FPhysInterface_Chaos::GetMass_AssumesLocked(const FPhysicsActorHandle& InActorReference)
 {
-	// #todo : Implement
-	return 1.f;
+	if (const Chaos::TPBDRigidParticle<float,3>* RigidParticle = InActorReference->AsDynamic())
+	{
+		return RigidParticle->M();
+	}
+	return 0.f;
 }
 
 void FPhysInterface_Chaos::SetSendsSleepNotifies_AssumesLocked(const FPhysicsActorHandle& InActorReference, bool bSendSleepNotifies)
@@ -444,18 +450,32 @@ void FPhysInterface_Chaos::SetSleepEnergyThreshold_AssumesLocked(const FPhysicsA
 {
 }
 
-void FPhysInterface_Chaos::SetMass_AssumesLocked(const FPhysicsActorHandle& InHandle, float InMass)
+void FPhysInterface_Chaos::SetMass_AssumesLocked(const FPhysicsActorHandle& InActorReference, float InMass)
 {
-	// #todo : Implement
-	//LocalParticles.M(Index) = InMass;
+	if (Chaos::TPBDRigidParticle<float, 3 >* RigidParticle = InActorReference->AsDynamic())
+	{
+		RigidParticle->SetM(InMass);
+		if (ensure(!FMath::IsNearlyZero(InMass)))
+		{
+			RigidParticle->SetInvM(1./InMass);
+		}
+		else
+		{
+			RigidParticle->SetInvM(0);
+		}
+	}
 }
 
-void FPhysInterface_Chaos::SetMassSpaceInertiaTensor_AssumesLocked(const FPhysicsActorHandle& InHandle, const FVector& InTensor)
+void FPhysInterface_Chaos::SetMassSpaceInertiaTensor_AssumesLocked(const FPhysicsActorHandle& InActorReference, const FVector& InTensor)
 {
-	// #todo : Implement
-	//LocalParticles.I(Index).M[0][0] = InTensor[0];
-    //LocalParticles.I(Index).M[1][1] = InTensor[1];
-    //LocalParticles.I(Index).M[2][2] = InTensor[2];
+	if (Chaos::TPBDRigidParticle<float, 3 >* RigidParticle = InActorReference->AsDynamic())
+	{
+		if(ensure(!FMath::IsNearlyZero(InTensor.X)) && ensure(!FMath::IsNearlyZero(InTensor.Y)) && ensure(!FMath::IsNearlyZero(InTensor.Z)) )
+		{
+			RigidParticle->SetI(Chaos::PMatrix<float,3,3>(InTensor.X, InTensor.Y, InTensor.Z));
+			RigidParticle->SetInvI(Chaos::PMatrix<float, 3, 3>(1./InTensor.X, 1./InTensor.Y, 1./InTensor.Z));
+		}
+	}
 }
 
 void FPhysInterface_Chaos::SetComLocalPose_AssumesLocked(const FPhysicsActorHandle& InHandle, const FTransform& InComLocalPose)
@@ -871,9 +891,10 @@ void FPhysInterface_Chaos::AddGeometry(FPhysicsActorHandle& InActor, const FGeom
 	{
 		for (uint32 i = 0; i < static_cast<uint32>(InParams.Geometry->SphereElems.Num()); ++i)
 		{
-			ensure(FMath::IsNearlyEqual(Scale[0],Scale[1]) && FMath::IsNearlyEqual(Scale[1],Scale[2]) );
-			const auto& CollisionSphere = InParams.Geometry->SphereElems[i];
-			auto ImplicitSphere = MakeUnique<Chaos::TSphere<float, 3>>(Chaos::TVector<float, 3>(0.f, 0.f, 0.f), CollisionSphere.Radius * Scale[0]);
+			const FKSphereElem ScaledSphereElem = InParams.Geometry->SphereElems[i].GetFinalScaled(Scale, InParams.LocalTransform);
+			const float UseRadius = FMath::Max(ScaledSphereElem.Radius, KINDA_SMALL_NUMBER);
+
+			auto ImplicitSphere = MakeUnique<Chaos::TSphere<float, 3>>(Chaos::TVector<float, 3>(0.f, 0.f, 0.f), UseRadius);
 			auto NewShape = NewShapeHelper(ImplicitSphere.Get());
 			Shapes.Emplace(MoveTemp(NewShape));
 
@@ -1442,21 +1463,70 @@ bool FPhysInterface_Chaos::Overlap_Geom(const FBodyInstance* InBodyInstance, con
 
 bool FPhysInterface_Chaos::GetSquaredDistanceToBody(const FBodyInstance* InInstance, const FVector& InPoint, float& OutDistanceSquared, FVector* OutOptPointOnBody)
 {
-	const FTransform BodyTM = InInstance->GetUnrealWorldTransform();
+	if (OutOptPointOnBody)
+	{
+		*OutOptPointOnBody = InPoint;
+		OutDistanceSquared = 0.f;
+	}
+
+	float ReturnDistance = -1.f;
+	float MinPhi = BIG_NUMBER;
+	bool bFoundValidBody = false;
+	bool bEarlyOut = true;
+
+	const FBodyInstance* UseBI = InInstance->WeldParent ? InInstance->WeldParent : InInstance;
+	const FTransform BodyTM = UseBI->GetUnrealWorldTransform();
 	const FVector LocalPoint = BodyTM.InverseTransformPositionNoScale(InPoint);
 
-    TArray<FPhysicsShapeReference_Chaos> OutShapes;
-    InInstance->GetAllShapes_AssumesLocked(OutShapes);
-    check(OutShapes.Num() == 1);
-    Chaos::TVector<float, 3> Normal;
-    const auto Phi = OutShapes[0].Shape->Geometry->PhiWithNormal(LocalPoint, Normal);
-    OutDistanceSquared = Phi * Phi;
-    if (OutOptPointOnBody)
-    {
-		const Chaos::TVector<float, 3> LocalClosestPoint = LocalPoint - Phi * Normal;
-        *OutOptPointOnBody = BodyTM.TransformPositionNoScale(LocalClosestPoint);
-    }
-    return true;
+	FPhysicsCommand::ExecuteRead(UseBI->ActorHandle, [&](const FPhysicsActorHandle& Actor)
+	{
+
+		bEarlyOut = false;
+
+		TArray<FPhysicsShapeReference_Chaos> Shapes;
+		InInstance->GetAllShapes_AssumesLocked(Shapes);
+		for (const FPhysicsShapeReference_Chaos& Shape : Shapes)
+		{
+			if (UseBI->IsShapeBoundToBody(Shape) == false)	//skip welded shapes that do not belong to us
+			{
+				continue;
+			}
+
+			ECollisionShapeType GeomType = FPhysicsInterface::GetShapeType(Shape);
+
+			if (GeomType == ECollisionShapeType::Trimesh)
+			{
+				// Type unsupported for this function, but some other shapes will probably work. 
+				continue;
+			}
+
+			bFoundValidBody = true;
+
+			Chaos::TVector<float, 3> Normal;
+			const float Phi = Shape.Shape->Geometry->PhiWithNormal(LocalPoint, Normal);
+			if (Phi <= 0)
+			{
+				break;
+			}
+			else if (Phi < MinPhi)
+			{
+				MinPhi = Phi;
+				OutDistanceSquared = Phi * Phi;
+				if (OutOptPointOnBody)
+				{
+					const Chaos::TVector<float, 3> LocalClosestPoint = LocalPoint - Phi * Normal;
+					*OutOptPointOnBody = BodyTM.TransformPositionNoScale(LocalClosestPoint);
+				}
+			}
+		}
+	});
+
+	if (!bFoundValidBody && !bEarlyOut)
+	{
+		UE_LOG(LogPhysics, Verbose, TEXT("GetDistanceToBody: Component (%s) has no simple collision and cannot be queried for closest point."), InInstance->OwnerComponent.Get() ? *(InInstance->OwnerComponent->GetPathName()) : TEXT("NONE"));
+	}
+
+	return bFoundValidBody;
 }
 
 template<typename AllocatorType>
@@ -1546,73 +1616,4 @@ const FQuat& FPhysicsShapeAdapter_Chaos::GetGeomOrientation() const
 {
 	return GeometryRotation;
 }
-
-FVector FindBoxOpposingNormal(const FLocationHit& Hit, const FVector& TraceDirectionDenorm, const FVector& InNormal)
-{
-	// We require normal info for our algorithm.
-	const bool bNormalData = !!(Hit.Flags & EHitFlags::Normal);
-	if (!bNormalData)
-	{
-		return InNormal;
-	}
-
-	ensure(Hit.Shape->Geometry->GetType() == Chaos::ImplicitObjectType::Box);
-	const FTransform LocalToWorld(Hit.Actor->R(), Hit.Actor->X());
-
-	// Find which faces were included in the contact normal, and for multiple faces, use the one most opposing the sweep direction.
-	const FVector ContactNormalLocal = LocalToWorld.InverseTransformVectorNoScale(Hit.WorldNormal);
-	const float* ContactNormalLocalPtr = &ContactNormalLocal.X;
-	const float* TraceDirDenormWorldPtr = &TraceDirectionDenorm.X;
-	const FVector TraceDirDenormLocal = LocalToWorld.InverseTransformVectorNoScale(TraceDirectionDenorm);
-	const float* TraceDirDenormLocalPtr = &TraceDirDenormLocal.X;
-
-	FVector BestLocalNormal(ContactNormalLocal);
-	float* BestLocalNormalPtr = &BestLocalNormal.X;
-	float BestOpposingDot = FLT_MAX;
-
-	for (int32 i = 0; i < 3; i++)
-	{
-		// Select axis of face to compare to, based on normal.
-		if (ContactNormalLocalPtr[i] > KINDA_SMALL_NUMBER)
-		{
-			const float TraceDotFaceNormal = TraceDirDenormLocalPtr[i]; // TraceDirDenormLocal.dot(BoxFaceNormal)
-			if (TraceDotFaceNormal < BestOpposingDot)
-			{
-				BestOpposingDot = TraceDotFaceNormal;
-				BestLocalNormal = FVector(0.f);
-				BestLocalNormalPtr[i] = 1.f;
-			}
-		}
-		else if (ContactNormalLocalPtr[i] < -KINDA_SMALL_NUMBER)
-		{
-			const float TraceDotFaceNormal = -TraceDirDenormLocalPtr[i]; // TraceDirDenormLocal.dot(BoxFaceNormal)
-			if (TraceDotFaceNormal < BestOpposingDot)
-			{
-				BestOpposingDot = TraceDotFaceNormal;
-				BestLocalNormal = FVector(0.f);
-				BestLocalNormalPtr[i] = -1.f;
-			}
-		}
-	}
-
-	// Fill in result
-	const FVector WorldNormal = LocalToWorld.TransformVectorNoScale(BestLocalNormal);
-	return WorldNormal;
-}
-
-FVector FindHeightFieldOpposingNormal(const FLocationHit& PHit, const FVector& TraceDirectionDenorm, const FVector& InNormal)
-{
-	//todo: implement
-	return FVector(0, 0, 1);
-}
-FVector FindConvexMeshOpposingNormal(const FLocationHit& PHit, const FVector& TraceDirectionDenorm, const FVector& InNormal)
-{
-	//todo: implement
-	return FVector(0, 0, 1);
-}
-FVector FindTriMeshOpposingNormal(const FLocationHit& PHit, const FVector& TraceDirectionDenorm, const FVector& InNormal)
-{
-	//todo: implement
-	return FVector(0, 0, 1);
-}
 #endif
diff --git a/Engine/Source/Runtime/Engine/Private/PhysicsEngine/Experimental/PhysScene_Chaos.cpp b/Engine/Source/Runtime/Engine/Private/PhysicsEngine/Experimental/PhysScene_Chaos.cpp
--- a/Engine/Source/Runtime/Engine/Private/PhysicsEngine/Experimental/PhysScene_Chaos.cpp
+++ b/Engine/Source/Runtime/Engine/Private/PhysicsEngine/Experimental/PhysScene_Chaos.cpp
@@ -26,9 +26,9 @@
 
 #include "PhysicsProxy/FieldSystemPhysicsProxy.h"
 #include "PhysicsProxy/GeometryCollectionPhysicsProxy.h"
+#include "PhysicsProxy/SingleParticlePhysicsProxy.h"
 #include "PhysicsProxy/SkeletalMeshPhysicsProxy.h"
 #include "PhysicsProxy/StaticMeshPhysicsProxy.h"
-#include "Chaos/Framework/SingleParticlePhysicsProxy.h"
 #include "Chaos/UniformGrid.h"
 #include "Chaos/BoundingVolume.h"
 #include "Chaos/ISpatialAcceleration.h"
@@ -583,6 +583,15 @@ void FPhysScene_Chaos::AddObject(UPrimitiveComponent* Component, FFieldSystemPhy
 #endif
 }
 
+void FPhysScene_Chaos::RemoveActorFromAccelerationStructure(FPhysicsActorHandle& Actor)
+{
+	if (SolverAccelerationStructure)
+	{
+		Chaos::TAccelerationStructureHandle<float, 3> AccelerationHandle(Actor);
+		SolverAccelerationStructure->RemoveElement(AccelerationHandle);
+	}
+}
+
 template<typename ObjectType>
 void RemovePhysicsProxy(ObjectType* InObject, Chaos::FPhysicsSolver* InSolver, FChaosSolversModule* InModule)
 {
@@ -1104,6 +1113,9 @@ void FPhysScene_ChaosInterface::StartFrame()
 	{
 	case EChaosThreadingMode::SingleThread:
 	{
+		OnPhysScenePreTick.Broadcast(this, MDeltaTime);
+		OnPhysSceneStep.Broadcast(this, MDeltaTime);
+
 		// Here we can directly tick the scene. Single threaded mode doesn't buffer any commands
 		// that would require pumping here - everything is done on demand.
 		Scene.Tick(MDeltaTime);
@@ -1113,6 +1125,9 @@ void FPhysScene_ChaosInterface::StartFrame()
 	{
 		check(!CompletionEvent.GetReference())
 
+		OnPhysScenePreTick.Broadcast(this, MDeltaTime);
+		OnPhysSceneStep.Broadcast(this, MDeltaTime);
+
 		{
 			const TArray<FPhysicsSolver*>& SolverList = SolverModule->GetSolvers();
 			TArray<FPhysicsSolver*> ActiveSolvers;
@@ -1182,6 +1197,8 @@ void FPhysScene_ChaosInterface::EndFrame(ULineBatchComponent* InLineBatcher)
 	{
 		SyncBodies();
 		Scene.GetSolver()->SyncEvents_GameThread();
+
+		OnPhysScenePostTick.Broadcast(this);
 	}
 	break;
 	case EChaosThreadingMode::TaskGraph:
@@ -1195,6 +1212,13 @@ void FPhysScene_ChaosInterface::EndFrame(ULineBatchComponent* InLineBatcher)
 		{
 			SCOPE_CYCLE_COUNTER(STAT_FlipResults);
 
+			//update external SQ structure
+			//for now just copy the whole thing, stomping any changes that came from GT
+			if (Chaos::FPhysicsSolver* SceneSolver = Scene.GetSolver())
+			{
+				Scene.SolverAccelerationStructure = SceneSolver->GetEvolution()->GetCollisionConstraints().CreateExternalAccelerationStructure();
+			}
+
 			const TArray<FPhysicsSolver*>& SolverList = SolverModule->GetSolvers();
 			TArray<FPhysicsSolver*> ActiveSolvers;
 			ActiveSolvers.Reserve(SolverList.Num());
@@ -1232,6 +1256,8 @@ void FPhysScene_ChaosInterface::EndFrame(ULineBatchComponent* InLineBatcher)
 				Solver->SyncEvents_GameThread();
 			}
 		}
+
+		OnPhysScenePostTick.Broadcast(this);
 	}
 		break;
 
diff --git a/Engine/Source/Runtime/Engine/Private/PhysicsEngine/ImmediatePhysics/ImmediatePhysicsChaos/ImmediatePhysicsActorHandle_Chaos.cpp b/Engine/Source/Runtime/Engine/Private/PhysicsEngine/ImmediatePhysics/ImmediatePhysicsChaos/ImmediatePhysicsActorHandle_Chaos.cpp
--- a/Engine/Source/Runtime/Engine/Private/PhysicsEngine/ImmediatePhysics/ImmediatePhysicsChaos/ImmediatePhysicsActorHandle_Chaos.cpp
+++ b/Engine/Source/Runtime/Engine/Private/PhysicsEngine/ImmediatePhysics/ImmediatePhysicsChaos/ImmediatePhysicsActorHandle_Chaos.cpp
@@ -413,6 +413,11 @@ namespace ImmediatePhysics_Chaos
 		return FLT_MAX;
 	}
 
+	const FVector& FActorHandle::GetCoMTranslation() const
+	{
+		return ActorToCoMTranslation;
+	}
+
 }
 
 #endif // INCLUDE_CHAOS
diff --git a/Engine/Source/Runtime/Engine/Private/PhysicsEngine/ImmediatePhysics/ImmediatePhysicsChaos/ImmediatePhysicsJointHandle_Chaos.cpp b/Engine/Source/Runtime/Engine/Private/PhysicsEngine/ImmediatePhysics/ImmediatePhysicsChaos/ImmediatePhysicsJointHandle_Chaos.cpp
--- a/Engine/Source/Runtime/Engine/Private/PhysicsEngine/ImmediatePhysics/ImmediatePhysicsChaos/ImmediatePhysicsJointHandle_Chaos.cpp
+++ b/Engine/Source/Runtime/Engine/Private/PhysicsEngine/ImmediatePhysics/ImmediatePhysicsChaos/ImmediatePhysicsJointHandle_Chaos.cpp
@@ -6,7 +6,7 @@
 
 #include "Physics/ImmediatePhysics/ImmediatePhysicsChaos/ImmediatePhysicsActorHandle_Chaos.h"
 
-#include "Chaos/PBDJointConstraints.h"
+#include "Chaos/PBD6DJointConstraints.h"
 
 #include "PhysicsEngine/ConstraintInstance.h"
 
@@ -14,25 +14,22 @@
 namespace ImmediatePhysics_Chaos
 {
 
-	FJointHandle::FJointHandle(Chaos::TPBDJointConstraints<FReal, Dimensions>* InConstraints, FConstraintInstance* ConstraintInstance, FActorHandle* Actor1, FActorHandle* Actor2)
+	FJointHandle::FJointHandle(FChaosConstraintContainer* InConstraints, FConstraintInstance* ConstraintInstance, FActorHandle* Actor1, FActorHandle* Actor2)
 		: Constraints(InConstraints)
-		, ConstraintIndex(INDEX_NONE)
 	{
-		FVector JointLocationLocal = ConstraintInstance->GetRefFrame(EConstraintFrame::Frame1).GetLocation();
-		FVector JointLocationWorld = Actor1->GetWorldTransform().TransformPosition(JointLocationLocal);
-		ConstraintIndex = Constraints->AddConstraint({ Actor1->ParticleHandle, Actor2->ParticleHandle }, JointLocationWorld);
-
-#if IMMEDIATEPHYSICS_CHAOS_TODO
-		// Chaos constraint handles
-		// Proper joint constraint
-#endif
+		// BodyInstance/PhysX has the constraint locations in actor-space, but we need them in Centre-of-Mass space
+		// @todo(ccaulfield): support CoM in Chaos particles
+		FTransform ConstraintFrame1 = ConstraintInstance->GetRefFrame(EConstraintFrame::Frame1);
+		FTransform ConstraintFrame2 = ConstraintInstance->GetRefFrame(EConstraintFrame::Frame2);
+		ConstraintFrame1.SetTranslation(ConstraintFrame1.GetTranslation() - Actor1->GetCoMTranslation());
+		ConstraintFrame2.SetTranslation(ConstraintFrame2.GetTranslation() - Actor2->GetCoMTranslation());
+
+		ConstraintHandle = Constraints->AddConstraint({ Actor1->ParticleHandle, Actor2->ParticleHandle }, { ConstraintFrame1, ConstraintFrame2 });
 	}
 
 	FJointHandle::~FJointHandle()
 	{
-#if IMMEDIATEPHYSICS_CHAOS_TODO
-		// Remove joints
-#endif
+		ConstraintHandle->RemoveConstraint();
 	}
 
 }
diff --git a/Engine/Source/Runtime/Engine/Private/PhysicsEngine/ImmediatePhysics/ImmediatePhysicsChaos/ImmediatePhysicsSimulation_Chaos.cpp b/Engine/Source/Runtime/Engine/Private/PhysicsEngine/ImmediatePhysics/ImmediatePhysicsChaos/ImmediatePhysicsSimulation_Chaos.cpp
--- a/Engine/Source/Runtime/Engine/Private/PhysicsEngine/ImmediatePhysics/ImmediatePhysicsChaos/ImmediatePhysicsSimulation_Chaos.cpp
+++ b/Engine/Source/Runtime/Engine/Private/PhysicsEngine/ImmediatePhysics/ImmediatePhysicsChaos/ImmediatePhysicsSimulation_Chaos.cpp
@@ -9,7 +9,7 @@
 #include "Physics/ImmediatePhysics/ImmediatePhysicsStats.h"
 
 #include "Chaos/ParticleHandle.h"
-#include "Chaos/PBDJointConstraints.h"
+#include "Chaos/PBD6DJointConstraints.h"
 #include "Chaos/PBDRigidParticles.h"
 #include "Chaos/PBDRigidsEvolutionGBF.h"
 
@@ -29,10 +29,10 @@ namespace ImmediatePhysics_Chaos
 	{
 		using namespace Chaos;
 
-		const int DefaultSolverIterations = 5;
+		const int DefaultSolverIterations = 3;
 		Particles = MakeUnique<TPBDRigidsSOAs<FReal, Dimensions>>();
-		Joints = MakeUnique<TPBDJointConstraints<FReal, Dimensions>>();
-		JointsRule = MakeUnique<TPBDConstraintIslandRule<TPBDJointConstraints<FReal, Dimensions>, FReal, Dimensions>>(*Joints);
+		Joints = MakeUnique<TPBD6DJointConstraints<FReal, Dimensions>>();
+		JointsRule = MakeUnique<TPBDConstraintIslandRule<TPBD6DJointConstraints<FReal, Dimensions>, FReal, Dimensions>>(*Joints);
 		Evolution = MakeUnique<TPBDRigidsEvolutionGBF<FReal, Dimensions>>(*Particles.Get(), DefaultSolverIterations);
 
 		Evolution->AddConstraintRule(JointsRule.Get());
diff --git a/Engine/Source/Runtime/Engine/Private/RayTracingInstanceUtils.cpp b/Engine/Source/Runtime/Engine/Private/RayTracingInstanceUtils.cpp
new file mode 100644
--- /dev/null
+++ b/Engine/Source/Runtime/Engine/Private/RayTracingInstanceUtils.cpp
@@ -0,0 +1,29 @@
+// Copyright 1998-2019 Epic Games, Inc. All Rights Reserved.
+
+/*=============================================================================
+	RayTracingInstance.cpp: Helper functions for creating a ray tracing instance.
+=============================================================================*/
+
+#include "RayTracingInstanceUtils.h"
+#include "RayTracingInstance.h"
+
+#if RHI_RAYTRACING
+
+
+void AddOpaqueRaytracingInstance(const FMatrix& InstanceTransform, const FRayTracingGeometry* RayTracingGeometry, const uint32 Mask, TArray<FRayTracingInstance>& OutRayTracingInstances)
+{
+	FRayTracingInstance Instance;
+	Instance.Geometry = RayTracingGeometry;
+	Instance.InstanceTransforms.Add(InstanceTransform);
+	Instance.bForceOpaque = true;
+	Instance.Mask = 0;
+	if (Mask & RaytracingInstanceMask_Opaque)		Instance.Mask |= RAY_TRACING_MASK_OPAQUE;
+	if (Mask & RaytracingInstanceMask_Translucent)	Instance.Mask |= RAY_TRACING_MASK_TRANSLUCENT;
+	if (Mask & RaytracingInstanceMask_Shadow)		Instance.Mask |= RAY_TRACING_MASK_SHADOW;
+	if (Mask & RaytracingInstanceMask_ThinShadow)	Instance.Mask |= RAY_TRACING_MASK_THIN_SHADOW;
+	OutRayTracingInstances.Add(Instance);
+
+
+}
+
+#endif
diff --git a/Engine/Source/Runtime/Engine/Private/SkeletalMesh.cpp b/Engine/Source/Runtime/Engine/Private/SkeletalMesh.cpp
--- a/Engine/Source/Runtime/Engine/Private/SkeletalMesh.cpp
+++ b/Engine/Source/Runtime/Engine/Private/SkeletalMesh.cpp
@@ -81,7 +81,7 @@
 #include "Streaming/UVChannelDensity.h"
 #include "Misc/Paths.h"
 
-#include "ClothingAssetInterface.h"
+#include "ClothingAssetBase.h"
 
 #if WITH_EDITOR
 #include "ClothingAssetFactoryInterface.h"
@@ -424,13 +424,10 @@ void USkeletalMesh::AddClothingAsset(UClothingAssetBase* InNewAsset)
 #if WITH_EDITOR
 void USkeletalMesh::RemoveClothingAsset(int32 InLodIndex, int32 InSectionIndex)
 {
-	UClothingAssetBase* Asset = GetSectionClothingAsset(InLodIndex, InSectionIndex);
-
-	if(Asset)
+	if(UClothingAssetBase* Asset = GetSectionClothingAsset(InLodIndex, InSectionIndex))
 	{
 		Asset->UnbindFromSkeletalMesh(this, InLodIndex);
 		MeshClothingAssets.Remove(Asset);
-
 		OnClothingChange.Broadcast();
 	}
 }
@@ -453,9 +450,9 @@ UClothingAssetBase* USkeletalMesh::GetSectionClothingAsset(int32 InLodIndex, int
 				{
 					UClothingAssetBase** FoundAsset = MeshClothingAssets.FindByPredicate([&](UClothingAssetBase* InAsset)
 					{
-						return InAsset->GetAssetGuid() == ClothingAssetGuid;
+						return InAsset && InAsset->GetAssetGuid() == ClothingAssetGuid;
 					});
-					
+
 					return FoundAsset ? *FoundAsset : nullptr;
 				}
 			}
@@ -482,7 +479,7 @@ const UClothingAssetBase* USkeletalMesh::GetSectionClothingAsset(int32 InLodInde
 				{
 					UClothingAssetBase* const* FoundAsset = MeshClothingAssets.FindByPredicate([&](UClothingAssetBase* InAsset)
 					{
-						return InAsset->GetAssetGuid() == ClothingAssetGuid;
+						return InAsset && InAsset->GetAssetGuid() == ClothingAssetGuid;
 					});
 
 					return FoundAsset ? *FoundAsset : nullptr;
@@ -503,7 +500,7 @@ UClothingAssetBase* USkeletalMesh::GetClothingAsset(const FGuid& InAssetGuid) co
 
 	UClothingAssetBase* const* FoundAsset = MeshClothingAssets.FindByPredicate([&](UClothingAssetBase* CurrAsset)
 	{
-		return CurrAsset->GetAssetGuid() == InAssetGuid;
+		return CurrAsset && CurrAsset->GetAssetGuid() == InAssetGuid;
 	});
 
 	return FoundAsset ? *FoundAsset : nullptr;
@@ -511,12 +508,7 @@ UClothingAssetBase* USkeletalMesh::GetClothingAsset(const FGuid& InAssetGuid) co
 
 int32 USkeletalMesh::GetClothingAssetIndex(UClothingAssetBase* InAsset) const
 {
-	if(!InAsset)
-	{
-		return INDEX_NONE;
-	}
-
-	return GetClothingAssetIndex(InAsset->GetAssetGuid());
+	return InAsset ? GetClothingAssetIndex(InAsset->GetAssetGuid()) : INDEX_NONE;
 }
 
 int32 USkeletalMesh::GetClothingAssetIndex(const FGuid& InAssetGuid) const
@@ -524,12 +516,12 @@ int32 USkeletalMesh::GetClothingAssetIndex(const FGuid& InAssetGuid) const
 	const int32 NumAssets = MeshClothingAssets.Num();
 	for(int32 SearchIndex = 0; SearchIndex < NumAssets; ++SearchIndex)
 	{
-		if(MeshClothingAssets[SearchIndex]->GetAssetGuid() == InAssetGuid)
+		if(MeshClothingAssets[SearchIndex] && 
+		   MeshClothingAssets[SearchIndex]->GetAssetGuid() == InAssetGuid)
 		{
 			return SearchIndex;
 		}
 	}
-
 	return INDEX_NONE;
 }
 
@@ -599,12 +591,9 @@ void USkeletalMesh::GetClothingAssetsInUse(TArray<UClothingAssetBase*>& OutCloth
 			for(int32 SectionIdx = 0; SectionIdx < NumSections; ++SectionIdx)
 			{
 				FSkelMeshRenderSection& Section = LodData.RenderSections[SectionIdx];
-
 				if(Section.ClothingData.AssetGuid.IsValid())
 				{
-					UClothingAssetBase* Asset = GetClothingAsset(Section.ClothingData.AssetGuid);
-					
-					if(Asset)
+					if(UClothingAssetBase* Asset = GetClothingAsset(Section.ClothingData.AssetGuid))
 					{
 						OutClothingAssets.AddUnique(Asset);
 					}
@@ -1825,7 +1814,7 @@ void USkeletalMesh::RemoveLegacyClothingSections()
 
 						Section.CorrespondClothAssetIndex = MeshClothingAssets.IndexOfByPredicate([&Section](const UClothingAssetBase* CurrAsset) 
 						{
-							return CurrAsset->GetAssetGuid() == Section.ClothingData.AssetGuid;
+							return CurrAsset && CurrAsset->GetAssetGuid() == Section.ClothingData.AssetGuid;
 						});
 
 						Section.BoneMap = DuplicatedSection.BoneMap;
@@ -2057,7 +2046,8 @@ void USkeletalMesh::PostLoad()
 	{
 		for(UClothingAssetBase* ClothingAsset : MeshClothingAssets)
 		{
-			ClothingAsset->InvalidateCachedData();
+			if(ClothingAsset) 
+				ClothingAsset->InvalidateCachedData();
 		}
 	}
 
diff --git a/Engine/Source/Runtime/Engine/Private/SkeletalMeshComponentPhysics.cpp b/Engine/Source/Runtime/Engine/Private/SkeletalMeshComponentPhysics.cpp
--- a/Engine/Source/Runtime/Engine/Private/SkeletalMeshComponentPhysics.cpp
+++ b/Engine/Source/Runtime/Engine/Private/SkeletalMeshComponentPhysics.cpp
@@ -2366,12 +2366,10 @@ void USkeletalMeshComponent::RecreateClothingActors()
 			for(int32 BaseAssetIndex = 0; BaseAssetIndex < NumMeshAssets; ++BaseAssetIndex)
 			{
 				UClothingAssetBase* Asset = SkeletalMesh->MeshClothingAssets[BaseAssetIndex];
-
-				if(!AssetsInUse.Contains(Asset))
+				if(!Asset || !AssetsInUse.Contains(Asset))
 				{
 					continue;
 				}
-
 				ClothingSimulation->CreateActor(this, Asset, BaseAssetIndex);
 			}
 
diff --git a/Engine/Source/Runtime/Engine/Private/UnrealEngine.cpp b/Engine/Source/Runtime/Engine/Private/UnrealEngine.cpp
--- a/Engine/Source/Runtime/Engine/Private/UnrealEngine.cpp
+++ b/Engine/Source/Runtime/Engine/Private/UnrealEngine.cpp
@@ -2363,6 +2363,8 @@ void UEngine::InitializeObjectReferences()
 	{
 		// Materials that are needed in-game if debug viewmodes are allowed
 		LoadSpecialMaterial(TEXT("WireframeMaterialName"), WireframeMaterialName, WireframeMaterial, true);
+		LoadSpecialMaterial(TEXT("HairDefaultMaterialName"), HairDefaultMaterialName, HairDefaultMaterial, true);
+		LoadSpecialMaterial(TEXT("HairDebugMaterialName"), HairDebugMaterialName, HairDebugMaterial, true);
 		LoadSpecialMaterial(TEXT("LevelColorationLitMaterialName"), LevelColorationLitMaterialName, LevelColorationLitMaterial, true);
 		LoadSpecialMaterial(TEXT("LevelColorationUnlitMaterialName"), LevelColorationUnlitMaterialName, LevelColorationUnlitMaterial, true);
 		LoadSpecialMaterial(TEXT("LightingTexelDensityName"), LightingTexelDensityName, LightingTexelDensityMaterial, false);
diff --git a/Engine/Source/Runtime/Engine/Public/CollisionQueryParams.h b/Engine/Source/Runtime/Engine/Public/CollisionQueryParams.h
--- a/Engine/Source/Runtime/Engine/Public/CollisionQueryParams.h
+++ b/Engine/Source/Runtime/Engine/Public/CollisionQueryParams.h
@@ -114,6 +114,12 @@ public:
 		bComponentListUnique = true;
 	}
 
+	/** Clears the set of actors to ignore during the trace. */
+	void ClearIgnoredActors()
+	{
+		IgnoreActors.Reset();
+	}
+
 	/**
 	 * Set the number of ignored components in the list. Uniqueness is not changed, it operates on the current state (unique or not).
 	 * Useful for temporarily adding some, then restoring to a previous size. NewNum must be <= number of current components for there to be any effect.
