Merging //UE4/Main@9388041 to Dev-RenderPlat-Staging

diff --git a/Engine/Source/Runtime/ClothingSystemRuntimeCommon/Public/ClothingAsset.h b/Engine/Source/Runtime/ClothingSystemRuntimeCommon/Public/ClothingAsset.h
new file mode 100644
--- /dev/null
+++ b/Engine/Source/Runtime/ClothingSystemRuntimeCommon/Public/ClothingAsset.h
@@ -0,0 +1,188 @@
+// Copyright 1998-2019 Epic Games, Inc. All Rights Reserved.
+#pragma once
+
+#include "ClothingAssetBase.h"
+#include "ClothConfig.h"
+#include "ClothLODData.h"
+#include "ClothingSimulationInteractor.h"
+
+#include "ClothingAsset.generated.h"
+
+DECLARE_LOG_CATEGORY_EXTERN(LogClothingAsset, Log, All);
+
+class UPhysicsAsset;
+struct FSkelMeshSection;
+
+namespace ClothingAssetUtils
+{
+	/**
+	 * Helper struct to hold binding information on a clothing asset, used to 
+	 * enumerate all of the bindings on a skeletal mesh with 
+	 * \c GetMeshClothingAssetBindings() below.
+	 */
+	struct FClothingAssetMeshBinding
+	{
+		UClothingAssetCommon* Asset;
+		int32 LODIndex;
+		int32 SectionIndex;
+		int32 AssetInternalLodIndex;
+	};
+
+	/**
+	 * Given a skeletal mesh, find all of the currently bound clothing assets and their binding information
+	 * @param InSkelMesh - The skeletal mesh to search
+	 * @param OutBindings - The list of bindings to write to
+	 */
+	void CLOTHINGSYSTEMRUNTIMECOMMON_API 
+	GetMeshClothingAssetBindings(
+		USkeletalMesh* InSkelMesh, 
+		TArray<FClothingAssetMeshBinding>& OutBindings);
+	
+	/**
+	 * Similar to above, but only inspects the specified LOD.
+	 */
+	void CLOTHINGSYSTEMRUNTIMECOMMON_API 
+	GetMeshClothingAssetBindings(
+		USkeletalMesh* InSkelMesh, 
+		TArray<FClothingAssetMeshBinding>& OutBindings, 
+		int32 InLodIndex);
+
+#if WITH_EDITOR
+	/**
+	 * Clears the clothing tracking struct of a section.
+	 */
+	CLOTHINGSYSTEMRUNTIMECOMMON_API void 
+	ClearSectionClothingData(FSkelMeshSection& InSection);
+#endif
+}
+
+/**
+ * Custom data wrapper for clothing assets.
+ * If writing a new clothing asset importer, creating a new derived custom data 
+ * is how to store importer (and possibly simulation) data that importer will 
+ * create. This needs to be set to the \c CustomData member on the asset your 
+ * factory creates.
+ *
+ * Testing whether a UClothingAssetCommon was made from a custom plugin can be 
+ * achieved with:
+ * \code if(AssetPtr->CustomData->IsA(UMyCustomData::StaticClass())) \endcode
+ */
+UCLASS(abstract, MinimalAPI)
+class UClothingAssetCustomData : public UObject
+{
+	GENERATED_BODY()
+public:
+	virtual void BindToSkeletalMesh(USkeletalMesh* InSkelMesh, int32 InMeshLodIndex, int32 InSectionIndex, int32 InAssetLodIndex)
+	{}
+};
+
+/**
+ * Implementation of non-solver specific, but common Engine related functionality.
+ *
+ * Solver specific implementations may wish to override this class to construct
+ * their own default instances of child classes, such as \c ClothSimConfig and 
+ * \c CustomData, as well as override the \c AddNewLod() factory to build their 
+ * own implementation of \c UClothLODDataBase.
+ */
+UCLASS(hidecategories = Object, BlueprintType)
+class CLOTHINGSYSTEMRUNTIMECOMMON_API UClothingAssetCommon : public UClothingAssetBase
+{
+	GENERATED_BODY()
+public:
+
+	UClothingAssetCommon(const FObjectInitializer& ObjectInitializer);
+
+#if WITH_EDITOR
+
+	/**
+	 * Create weights for skinning the render mesh to our simulation mesh, and 
+	 * weights to drive our sim mesh from the skeleton.
+	 */
+	virtual bool BindToSkeletalMesh(USkeletalMesh* InSkelMesh, const int32 InMeshLodIndex, const int32 InSectionIndex, const int32 InAssetLodIndex, const bool bCallPostEditChange=true) override;
+
+	/**
+	 * Helper that invokes \c UnbindFromSkeletalMesh() for each avilable entry in 
+	 * \p InSkelMesh->GetImportedModel()'s LODModel.
+	 */
+	virtual void UnbindFromSkeletalMesh(USkeletalMesh* InSkelMesh) override;
+	virtual void UnbindFromSkeletalMesh(USkeletalMesh* InSkelMesh, const int32 InMeshLodIndex) override;
+
+	void ReregisterComponentsUsingClothing();
+	void ForEachInteractorUsingClothing(TFunction<void (UClothingSimulationInteractor*)> Func);
+
+	/** 
+	 * Callback envoked after weights have been edited.
+	 * Calls \c PushWeightsToMesh() on each \c ClothLodData, and invalidates cached data. 
+	 */
+	virtual void ApplyParameterMasks();
+
+	/**
+	 *	Builds the LOD transition data.
+	 *	When we transition between LODs we skin the incoming mesh to the outgoing mesh
+	 *	in exactly the same way the render mesh is skinned to create a smooth swap
+	 */
+	void BuildLodTransitionData();
+
+#endif // WITH_EDITOR
+
+	/**
+	 * Rebuilds the \c UsedBoneIndices array from looking up the entries in the
+	 * \c UsedBoneNames array, in the \p InSkelMesh's reference skeleton.
+	 */
+	virtual void RefreshBoneMapping(USkeletalMesh* InSkelMesh) override;
+
+	/**
+	 * Calculates the prefered root bone for the simulation.
+	 */
+	void CalculateReferenceBoneIndex();
+
+	/**
+	 * Returns \c true if \p InLodIndex is a valid LOD id (index into \c ClothLodData).
+	 */
+	virtual bool IsValidLod(int32 InLodIndex) override;
+
+	/**
+	 * Returns the number of valid LOD's (length of the \c ClothLodData array).
+	 */
+	virtual int32 GetNumLods() override;
+
+	/**
+	 * Builds self collision data.
+	 * The default behavior is if the \c ClothSimConfig reports that self collisions
+	 * are enabled, it envokes \c BuildSelfCollisionData() on each \c ClothLodData's
+	 * \c PhysicalMeshData member.
+	 */
+	virtual void BuildSelfCollisionData() override;
+
+	// The physics asset to extract collisions from when building a simulation.
+	UPROPERTY(EditAnywhere, Category = Config)
+	UPhysicsAsset* PhysicsAsset;
+
+	// Parameters for how the cloth behaves.
+	UPROPERTY(EditAnywhere, Category = Config, NoClear, meta = (NoResetToDefault))
+	UClothConfigBase* ClothSimConfig;
+
+	// The actual asset data, listed by LOD.
+	UPROPERTY()
+	TArray<UClothLODDataBase*> ClothLodData;
+
+	// Tracks which clothing LOD each skel mesh LOD corresponds to (LodMap[SkelLod]=ClothingLod).
+	UPROPERTY()
+	TArray<int32> LodMap;
+
+	// List of bones this asset uses inside its parent mesh.
+	UPROPERTY()
+	TArray<FName> UsedBoneNames;
+
+	// List of the indices for the bones in UsedBoneNames, used for remapping.
+	UPROPERTY()
+	TArray<int32> UsedBoneIndices;
+
+	// Bone to treat as the root of the simulation space.
+	UPROPERTY()
+	int32 ReferenceBoneIndex;
+
+	// Custom data applied by the importer depending on where the asset was imported from.
+	UPROPERTY()
+	UClothingAssetCustomData* CustomData;
+};
\ No newline at end of file
diff --git a/Engine/Source/Runtime/ClothingSystemRuntimeCommon/Public/Utils/ClothingMeshUtils.h b/Engine/Source/Runtime/ClothingSystemRuntimeCommon/Public/Utils/ClothingMeshUtils.h
new file mode 100644
--- /dev/null
+++ b/Engine/Source/Runtime/ClothingSystemRuntimeCommon/Public/Utils/ClothingMeshUtils.h
@@ -0,0 +1,193 @@
+// Copyright 1998-2019 Epic Games, Inc. All Rights Reserved.
+
+#pragma once
+
+#include "SkeletalMeshTypes.h"
+
+DECLARE_LOG_CATEGORY_EXTERN(LogClothingMeshUtils, Log, All);
+
+class UClothPhysicalMeshDataBase;
+
+namespace ClothingMeshUtils
+{
+	struct CLOTHINGSYSTEMRUNTIMECOMMON_API ClothMeshDesc
+	{
+		ClothMeshDesc(TArrayView<const FVector> InPositions, TArrayView<const FVector> InNormals,  TArrayView<const uint32> InIndices)
+			: Positions(InPositions)
+			, Normals(InNormals)
+			, Indices(InIndices)
+		{}
+
+		bool HasValidMesh() const
+		{
+			return Positions.Num() == Normals.Num() && Indices.Num() % 3 == 0;
+		}
+
+		TArrayView<const FVector> Positions;
+		TArrayView<const FVector> Normals;
+		TArrayView<const uint32> Indices;
+	};
+
+	// Static method for calculating a skinned mesh result from source data
+	void CLOTHINGSYSTEMRUNTIMECOMMON_API
+	SkinPhysicsMesh(
+		const TArray<int32>& BoneMap, // UClothingAssetCommon::UsedBoneIndices
+		const UClothPhysicalMeshDataBase& InMesh, 
+		const FTransform& RootBoneTransform, 
+		const FMatrix* InBoneMatrices, 
+		const int32 InNumBoneMatrices, 
+		TArray<FVector>& OutPositions, 
+		TArray<FVector>& OutNormals);
+
+	/**
+	* Given mesh information for two meshes, generate a list of skinning data to embed mesh0 in mesh1
+	* @param OutSkinningData	- Final skinning data to map mesh0 to mesh1
+	* @param TargetMesh		- Mesh data for the mesh we are embedding
+	* @param TargetTangents		- Optional Tangents for the mesh we are embedding
+	* @param Mesh1Normals		- Vertex normals for Mesh1
+	* @param Mesh1Indices		- Triangle indices for Mesh1
+	*/
+	void CLOTHINGSYSTEMRUNTIMECOMMON_API GenerateMeshToMeshSkinningData(
+		TArray<FMeshToMeshVertData>& OutSkinningData,
+		const ClothMeshDesc& TargetMesh,
+		const TArray<FVector>* TargetTangents,
+		const ClothMeshDesc& SourceMesh);
+
+	/** 
+	 * Embeds a list of positions into a source mesh
+	 * @param SourceMesh The mesh to embed in
+	 * @param Positions The positions to embed in SourceMesh
+	 * @param OutEmbeddedPositions Embedded version of the original positions, a barycentric coordinate and distance along the normal of the triangle
+	 * @param OutSourceIndices Source index list for the embedded positions, 3 per position to denote the source triangle
+	 */
+	void CLOTHINGSYSTEMRUNTIMECOMMON_API GenerateEmbeddedPositions(
+		const ClothMeshDesc& SourceMesh, 
+		TArrayView<const FVector> Positions, 
+		TArray<FVector4>& OutEmbeddedPositions, 
+		TArray<int32>& OutSourceIndices);
+
+	/**
+	* Given a triangle ABC with normals at each vertex NA, NB and NC, get a barycentric coordinate
+	* and corresponding distance from the triangle encoded in an FVector4 where the components are
+	* (BaryX, BaryY, BaryZ, Dist)
+	* @param A		- Position of triangle vertex A
+	* @param B		- Position of triangle vertex B
+	* @param C		- Position of triangle vertex C
+	* @param NA	- Normal at vertex A
+	* @param NB	- Normal at vertex B
+	* @param NC	- Normal at vertex C
+	* @param Point	- Point to calculate Bary+Dist for
+	*/
+	FVector4 GetPointBaryAndDist(
+		const FVector& A,
+		const FVector& B,
+		const FVector& C,
+		const FVector& NA,
+		const FVector& NB,
+		const FVector& NC,
+		const FVector& Point);
+
+	/** 
+	 * Object used to map vertex parameters between two meshes using the
+	 * same barycentric mesh to mesh mapping data we use for clothing
+	 */
+	class CLOTHINGSYSTEMRUNTIMECOMMON_API FVertexParameterMapper
+	{
+	public:
+		FVertexParameterMapper() = delete;
+		FVertexParameterMapper(const FVertexParameterMapper& Other) = delete;
+
+		FVertexParameterMapper(TArrayView<const FVector> InMesh0Positions,
+			TArrayView<const FVector> InMesh0Normals,
+			TArrayView<const FVector> InMesh1Positions,
+			TArrayView<const FVector> InMesh1Normals,
+			TArrayView<const uint32> InMesh1Indices)
+			: Mesh0Positions(InMesh0Positions)
+			, Mesh0Normals(InMesh0Normals)
+			, Mesh1Positions(InMesh1Positions)
+			, Mesh1Normals(InMesh1Normals)
+			, Mesh1Indices(InMesh1Indices)
+		{
+
+		}
+
+		/** Generic mapping function, can be used to map any type with a provided callable */
+		template<typename T, typename Lambda>
+		void Map(TArrayView<const T>& SourceData, TArray<T>& DestData, const Lambda& Func)
+		{
+			// Enforce the interp func signature (returns T and takes a bary and 3 Ts)
+			// If you hit this then either the return type isn't T or your arguments aren't convertible to T
+			static_assert(TAreTypesEqual<T, typename TDecay<decltype(Func(DeclVal<FVector>(), DeclVal<T>(), DeclVal<T>(), DeclVal<T>()))>::Type>::Value, "Invalid Lambda signature passed to Map");
+
+			const int32 NumMesh0Positions = Mesh0Positions.Num();
+			const int32 NumMesh0Normals = Mesh0Normals.Num();
+
+			const int32 NumMesh1Positions = Mesh1Positions.Num();
+			const int32 NumMesh1Normals = Mesh1Normals.Num();
+			const int32 NumMesh1Indices = Mesh1Indices.Num();
+
+			// Validate mesh data
+			check(NumMesh0Positions == NumMesh0Normals);
+			check(NumMesh1Positions == NumMesh1Normals);
+			check(NumMesh1Indices % 3 == 0);
+			check(SourceData.Num() == NumMesh1Positions);
+
+			if(DestData.Num() != NumMesh0Positions)
+			{
+				DestData.Reset();
+				DestData.AddUninitialized(NumMesh0Positions);
+			}
+
+			ClothMeshDesc SourceMeshDesc(Mesh1Positions, Mesh1Normals, Mesh1Indices);
+
+			TArray<FVector4> EmbeddedPositions;
+			TArray<int32> SourceIndices;
+			ClothingMeshUtils::GenerateEmbeddedPositions(SourceMeshDesc, Mesh0Positions, EmbeddedPositions, SourceIndices);
+
+			for(int32 DestVertIndex = 0 ; DestVertIndex < NumMesh0Positions ; ++DestVertIndex)
+			{
+				// Truncate the distance from the position data
+				FVector Bary = EmbeddedPositions[DestVertIndex];
+
+				const int32 SourceTriBaseIdx = DestVertIndex * 3;
+				T A = SourceData[SourceIndices[SourceTriBaseIdx + 0]];
+				T B = SourceData[SourceIndices[SourceTriBaseIdx + 1]];
+				T C = SourceData[SourceIndices[SourceTriBaseIdx + 2]];
+
+				T& DestVal = DestData[DestVertIndex];
+
+				// If we're super close to a vertex just take it's value.
+				// Otherwise call the provided interp lambda
+				const static FVector OneVec(1.0f, 1.0f, 1.0f);
+				FVector DiffVec = OneVec - Bary;
+				if(FMath::Abs(DiffVec.X) <= SMALL_NUMBER)
+				{
+					DestVal = A;
+				}
+				else if(FMath::Abs(DiffVec.Y) <= SMALL_NUMBER)
+				{
+					DestVal = B;
+				}
+				else if(FMath::Abs(DiffVec.Z) <= SMALL_NUMBER)
+				{
+					DestVal = C;
+				}
+				else
+				{
+					DestData[DestVertIndex] = Func(Bary, A, B, C);
+				}
+			}
+		}
+
+		// Defined type mappings for brevity
+		void Map(TArrayView<const float> Source, TArray<float>& Dest);
+
+	private:
+
+		TArrayView<const FVector> Mesh0Positions;
+		TArrayView<const FVector> Mesh0Normals;
+		TArrayView<const FVector> Mesh1Positions;
+		TArrayView<const FVector> Mesh1Normals;
+		TArrayView<const uint32> Mesh1Indices;
+	};
+}
\ No newline at end of file
diff --git a/Engine/Source/Runtime/ClothingSystemRuntimeInterface/ClothingSystemRuntimeInterface.Build.cs b/Engine/Source/Runtime/ClothingSystemRuntimeInterface/ClothingSystemRuntimeInterface.Build.cs
--- a/Engine/Source/Runtime/ClothingSystemRuntimeInterface/ClothingSystemRuntimeInterface.Build.cs
+++ b/Engine/Source/Runtime/ClothingSystemRuntimeInterface/ClothingSystemRuntimeInterface.Build.cs
@@ -1,20 +1,20 @@
 // Copyright 1998-2019 Epic Games, Inc. All Rights Reserved.
 
 using UnrealBuildTool;
-using System.IO;
 
 public class ClothingSystemRuntimeInterface : ModuleRules
 {
 	public ClothingSystemRuntimeInterface(ReadOnlyTargetRules Target) : base(Target)
 	{
-        ShortName = "ClothSysRuntimeIntrfc";
+		ShortName = "ClothSysRuntimeIntrfc";
 
         PrivateIncludePaths.Add("Runtime/ClothingSystemRuntimeInterface/Private");
+        SetupModulePhysicsSupport(Target);
 
         PublicDependencyModuleNames.AddRange(
             new string[] {
                 "Core",
-                "CoreUObject",
+                "CoreUObject"
             }
         );
     }
diff --git a/Engine/Source/Runtime/ClothingSystemRuntimeInterface/Private/ClothConfig.cpp b/Engine/Source/Runtime/ClothingSystemRuntimeInterface/Private/ClothConfig.cpp
new file mode 100644
--- /dev/null
+++ b/Engine/Source/Runtime/ClothingSystemRuntimeInterface/Private/ClothConfig.cpp
@@ -0,0 +1,7 @@
+#include "ClothConfig.h"
+
+UClothConfigBase::UClothConfigBase()
+{}
+
+UClothConfigBase::~UClothConfigBase()
+{}
diff --git a/Engine/Source/Runtime/ClothingSystemRuntimeInterface/Private/ClothPhysicalMeshData.cpp b/Engine/Source/Runtime/ClothingSystemRuntimeInterface/Private/ClothPhysicalMeshData.cpp
new file mode 100644
--- /dev/null
+++ b/Engine/Source/Runtime/ClothingSystemRuntimeInterface/Private/ClothPhysicalMeshData.cpp
@@ -0,0 +1,62 @@
+#include "ClothPhysicalMeshData.h"
+
+UClothPhysicalMeshDataBase::UClothPhysicalMeshDataBase()
+	: NumFixedVerts(0)
+	, MaxBoneWeights(0)
+{}
+
+UClothPhysicalMeshDataBase::~UClothPhysicalMeshDataBase()
+{}
+
+void UClothPhysicalMeshDataBase::Reset(const int32 InNumVerts)
+{
+	Vertices.Reset();
+	Normals.Reset();
+#if WITH_EDITORONLY_DATA
+	VertexColors.Reset();
+#endif // #if WITH_EDITORONLY_DATA
+	InverseMasses.Reset();
+	BoneData.Reset();
+
+	Vertices.AddDefaulted(InNumVerts);
+	Normals.AddDefaulted(InNumVerts);
+#if WITH_EDITORONLY_DATA
+	VertexColors.AddDefaulted(InNumVerts);
+#endif //#if WITH_EDITORONLY_DATA
+
+	InverseMasses.AddDefaulted(InNumVerts);
+	BoneData.AddDefaulted(InNumVerts);
+
+	NumFixedVerts = 0;
+	MaxBoneWeights = 0;
+}
+
+TArray<float>* UClothPhysicalMeshDataBase::GetFloatArray(const uint32 Id) const
+{
+	check(IdToArray.Contains(Id));
+	return IdToArray[Id];
+}
+
+TArray<uint32> UClothPhysicalMeshDataBase::GetFloatArrayIds() const
+{
+	TArray<uint32> Keys; 
+	IdToArray.GetKeys(Keys);
+	return Keys;
+}
+
+TArray<TArray<float>*> UClothPhysicalMeshDataBase::GetFloatArrays() const
+{
+	TArray<TArray<float>*> Values;
+	IdToArray.GenerateValueArray(Values);
+	return Values;
+}
+
+void UClothPhysicalMeshDataBase::RegisterFloatArray(
+	const uint32 Id,
+	TArray<float> *Array)
+{
+	check(Id != INDEX_NONE);
+	check(Array != nullptr);
+	check(!IdToArray.Contains(Id) || IdToArray[Id] == Array);
+	IdToArray.Add(Id, Array);
+}
diff --git a/Engine/Source/Runtime/ClothingSystemRuntimeInterface/Public/ClothConfig.h b/Engine/Source/Runtime/ClothingSystemRuntimeInterface/Public/ClothConfig.h
new file mode 100644
--- /dev/null
+++ b/Engine/Source/Runtime/ClothingSystemRuntimeInterface/Public/ClothConfig.h
@@ -0,0 +1,45 @@
+#pragma once
+
+#include "UObject/ObjectMacros.h"
+#include "ClothConfig.generated.h"
+
+class UClothingAssetBase;
+
+/**
+ * Base class for simulator specific simulation controls.
+ */
+UCLASS(Abstract, DefaultToInstanced)
+class CLOTHINGSYSTEMRUNTIMEINTERFACE_API UClothConfigBase : public UObject
+{
+	GENERATED_BODY()
+public:
+	UClothConfigBase();
+	virtual ~UClothConfigBase();
+
+	virtual bool HasSelfCollision() const
+	{ unimplemented(); return false; }
+
+#if WITH_EDITOR
+	/** 
+	 * Callback invoked by \c UClothingAssetBase::PostEditChangeChainProperty(). 
+	 * Override to add custom handling of editor events.
+	 */
+	virtual bool PostEditChangeChainPropertyCallback(
+		UClothingAssetBase* Asset,
+		FPropertyChangedChainEvent& Event)
+	{
+		return false; 
+	}
+
+	/**
+	 * Callback invoked by \c UClothingAssetBase::OnFinishedChangingClothPropertiesCallback().
+	 * Override to add custom handling of editor events.
+	 */
+	virtual bool OnFinishedChangingClothingPropertiesCallback(
+		UClothingAssetBase* Asset,
+		const FPropertyChangedEvent& Event)
+	{
+		return false;
+	}
+#endif
+};
diff --git a/Engine/Source/Runtime/ClothingSystemRuntimeInterface/Public/ClothPhysicalMeshData.h b/Engine/Source/Runtime/ClothingSystemRuntimeInterface/Public/ClothPhysicalMeshData.h
new file mode 100644
--- /dev/null
+++ b/Engine/Source/Runtime/ClothingSystemRuntimeInterface/Public/ClothPhysicalMeshData.h
@@ -0,0 +1,114 @@
+// Copyright 1998-2019 Epic Games, Inc. All Rights Reserved.
+#pragma once
+
+#include "UObject/Object.h"
+
+#include "ClothVertBoneData.h"
+#include "Containers/Array.h"
+#include "Containers/Map.h"
+
+#include "ClothPhysicalMeshData.generated.h"
+
+class UClothConfigBase;
+
+/**
+ * Simulation mesh points, topology, and spatial parameters defined on that 
+ * topology.
+ *
+ * Created curing asset import or created from a skeletal mesh.
+ */
+UCLASS()
+class CLOTHINGSYSTEMRUNTIMEINTERFACE_API UClothPhysicalMeshDataBase : public UObject
+{
+	GENERATED_BODY()
+public:
+	UClothPhysicalMeshDataBase();
+	virtual ~UClothPhysicalMeshDataBase();
+
+	virtual void 
+	Reset(const int32 InNumVerts);
+
+	// TODO: Ryan - Rename to ClearWeightedParameters()?
+	/** Clear out any target properties in this physical mesh. */
+	virtual void 
+	ClearParticleParameters()
+	{}
+
+	/** Callback invoked from \c UClothingAssetBase::BuildSelfCollisionData(), 
+	 * implement in derived classes. 
+	 */
+	virtual void 
+	BuildSelfCollisionData(const UClothConfigBase* ClothConfig)
+	{ unimplemented(); }
+
+	/** Retrieve a registered vertex weight array by unique @param Id. */
+	TArray<float>* 
+	GetFloatArray(const uint32 Id) const;
+
+	/** Get ids for all registered weight arrays. */
+	TArray<uint32>
+	GetFloatArrayIds() const;
+
+	/** Get all registered weight arrays. */
+	TArray<TArray<float>*>
+	GetFloatArrays() const;
+
+	/** Returns an \c Enum mapping float array Id's to names used in the UI. */
+	virtual UEnum* 
+	GetFloatArrayTargets() const
+	{ unimplemented(); return nullptr; }
+
+
+	virtual bool IsFullyKinematic() const
+	{ unimplemented(); return true; }
+	virtual bool IsFixed(const uint16 X, const float Threshold=0.1f) const
+	{ unimplemented(); return false; }
+	virtual bool IsFixed(const uint16 X, const uint16 Y, const uint16 Z, const float Threshold=0.1f) const
+	{ unimplemented(); return false; }
+
+protected:
+	/** Register an @param Array keyed by a unique @param Id. */
+	void RegisterFloatArray(const uint32 Id, TArray<float> *Array);
+
+public:
+	// Positions of each simulation vertex
+	UPROPERTY(EditAnywhere, Category = SimMesh)
+	TArray<FVector> Vertices;
+
+	// Normal at each vertex
+	UPROPERTY(EditAnywhere, Category = SimMesh)
+	TArray<FVector> Normals;
+
+#if WITH_EDITORONLY_DATA
+	// Color at each vertex
+	UPROPERTY(EditAnywhere, Category = SimMesh)
+	TArray<FColor> VertexColors;
+#endif // WITH_EDITORONLY_DATA
+
+	// Indices of the simulation mesh triangles
+	UPROPERTY(EditAnywhere, Category = SimMesh)
+	TArray<uint32> Indices;
+
+	// Inverse mass for each vertex in the physical mesh
+	UPROPERTY(EditAnywhere, Category = SimMesh)
+	TArray<float> InverseMasses;
+
+	// Indices and weights for each vertex, used to skin the mesh to create the reference pose
+	UPROPERTY(EditAnywhere, Category = SimMesh)
+	TArray<FClothVertBoneData> BoneData;
+
+	// Number of fixed verts in the simulation mesh (fixed verts are just skinned and do not simulate)
+	UPROPERTY(EditAnywhere, Category = SimMesh)
+	int32 NumFixedVerts;
+
+	// Maximum number of bone weights of any vetex
+	UPROPERTY(EditAnywhere, Category = SimMesh)
+	int32 MaxBoneWeights;
+
+	// Valid indices to use for self collisions (reduced set of Indices)
+	UPROPERTY(EditAnywhere, Category = SimMesh)
+	TArray<uint32> SelfCollisionIndices;
+
+private:
+	TMap<uint32, TArray<float>*> IdToArray;
+};
diff --git a/Engine/Source/Runtime/ClothingSystemRuntimeInterface/Public/ClothingAssetBase.h b/Engine/Source/Runtime/ClothingSystemRuntimeInterface/Public/ClothingAssetBase.h
--- a/Engine/Source/Runtime/ClothingSystemRuntimeInterface/Public/ClothingAssetBase.h
+++ b/Engine/Source/Runtime/ClothingSystemRuntimeInterface/Public/ClothingAssetBase.h
@@ -23,7 +23,7 @@ public:
 	* @param InSubmeshIdx Submesh in this asset to replace section with
 	* @param InAssetLodIndex Internal clothing LOD to use
 	*/
-	virtual bool BindToSkeletalMesh(USkeletalMesh* InSkelMesh, int32 InMeshLodIndex, int32 InSectionIndex, int32 InAssetLodIndex, bool bCallPostEditChange = true)
+	virtual bool BindToSkeletalMesh(USkeletalMesh* InSkelMesh, const int32 InMeshLodIndex, const int32 InSectionIndex, const int32 InAssetLodIndex, const bool bCallPostEditChange = true)
 	PURE_VIRTUAL(UClothingAssetBase::BindToSkeletalMesh, return false;);
 
 	/**
@@ -38,7 +38,7 @@ public:
 	* @param InSkelMesh skeletal mesh to remove from
 	* @param InMeshLodIndex Mesh LOD to remove this asset from (could still be bound to other LODs)
 	*/
-	virtual void UnbindFromSkeletalMesh(USkeletalMesh* InSkelMesh, int32 InMeshLodIndex)
+	virtual void UnbindFromSkeletalMesh(USkeletalMesh* InSkelMesh, const int32 InMeshLodIndex)
 	PURE_VIRTUAL(UClothingAssetBase::UnbindFromSkeletalMesh,);
 
 	/** 
@@ -65,10 +65,26 @@ public:
 	virtual int32 GetNumLods()
 	PURE_VIRTUAL(UClothingAssetBase::GetNumLods(), return 0;);
 
+	/** Add a new LOD class instance. */
+	virtual int32 AddNewLod()
+	PURE_VIRTUAL(UClothingAssetBase::AddNewLod(), return INDEX_NONE;);
+
+	/** Called from \c FPersonaMeshDetails::OnFinishedChangingClothingProperties(). */
+	virtual void PostPropertyChangeCb(const FPropertyChangedEvent& Event)
+	PURE_VIRTUAL(UClothingAssetBase::PostPropertyChangeCb, );
+
 	/** Builds self collision data */
 	virtual void BuildSelfCollisionData()
 	PURE_VIRTUAL(UClothingAssetBase::BuildSelfCollisionData(), );
 
+	/**
+	*	Builds the LOD transition data
+	*	When we transition between LODs we skin the incoming mesh to the outgoing mesh
+	*	in exactly the same way the render mesh is skinned to create a smooth swap
+	*/
+	virtual void BuildLodTransitionData()
+	PURE_VIRTUAL(UClothingAssetBase::BuildLodTransitionData(), );
+
 	/** Get the guid identifying this asset */
 	const FGuid& GetAssetGuid() const
 	{
diff --git a/Engine/Source/Runtime/ClothingSystemRuntimeInterface/Public/ClothingSimulationContext.h b/Engine/Source/Runtime/ClothingSystemRuntimeInterface/Public/ClothingSimulationContext.h
--- a/Engine/Source/Runtime/ClothingSystemRuntimeInterface/Public/ClothingSimulationContext.h
+++ b/Engine/Source/Runtime/ClothingSystemRuntimeInterface/Public/ClothingSimulationContext.h
@@ -14,7 +14,7 @@ public:
 	virtual ~IClothingSimulationContext();
 };
 
-/** Base simlation data that just about every simulation would need. */
+/** Base simulation data that just about every simulation would need. */
 class CLOTHINGSYSTEMRUNTIMEINTERFACE_API FClothingSimulationContextBase : public IClothingSimulationContext
 {
 public:
diff --git a/Engine/Source/Runtime/ClothingSystemRuntimeInterface/Public/PointWeightMap.h b/Engine/Source/Runtime/ClothingSystemRuntimeInterface/Public/PointWeightMap.h
new file mode 100644
--- /dev/null
+++ b/Engine/Source/Runtime/ClothingSystemRuntimeInterface/Public/PointWeightMap.h
@@ -0,0 +1,105 @@
+// Copyright 1998-2019 Epic Games, Inc. All Rights Reserved.
+#pragma once
+
+#include "UObject/ObjectMacros.h"
+#include "Math/NumericLimits.h"
+#include "PointWeightMap.generated.h"
+
+/** 
+ * A mask is simply some storage for a physical mesh parameter painted onto clothing.
+ * Used in the editor for users to paint onto and then target to a parameter, which
+ * is then later applied to a phys mesh
+ */
+USTRUCT()
+struct CLOTHINGSYSTEMRUNTIMEINTERFACE_API FPointWeightMap
+{
+	GENERATED_BODY();
+
+	FPointWeightMap()
+		: Name(NAME_None)
+		, CurrentTarget(0) // 0 = NONE
+		, bEnabled(false)
+	{}
+	~FPointWeightMap()
+	{}
+
+	/** Initialize the weight map. */
+	void
+	Initialize(const int32 NumPoints)
+	{
+		Values.Reset(NumPoints);
+		Values.AddZeroed(NumPoints);
+		bEnabled = false; // Not sure why we do this...
+	}
+
+	/** Copies from \p SourceValues to \c Values. */
+	void
+	CopyFrom(const TArray<float>& SourceValues, uint8 Target)
+	{ Initialize(SourceValues.Num()); Values = SourceValues; CurrentTarget = Target; }
+	void
+	CopyFrom(const TArray<float>* SourceValues, uint8 Target)
+	{ if (SourceValues) CopyFrom(*SourceValues, Target); }
+
+	/** Copies from \c Values to \p TargetValues. */
+	void
+	CopyTo(TArray<float>& TargetValues) const
+	{ TargetValues = Values; }
+	void
+	CopyTo(TArray<float>* TargetValues) const
+	{ if (TargetValues) CopyTo(*TargetValues); }
+
+	/** 
+	 * Set a value in the mask
+	 * @param InVertexIndex the value/vertex index to set
+	 * @param InValue the value to set
+	 */
+	void 
+	SetValue(int32 InVertexIndex, float InValue)
+	{ if (Values.IsValidIndex(InVertexIndex)) Values[InVertexIndex] = InValue; }
+
+	/** 
+	 * Get a value from the mask
+	 * @param InVertexIndex the value/vertex index to retrieve
+	 */
+	float 
+	GetValue(int32 InVertexIndex) const
+	{ return Values.IsValidIndex(InVertexIndex) ? Values[InVertexIndex] : 0.0f; }
+	
+	/** 
+	* Read only version of the array holding the mask values
+	*/
+	const TArray<float>& 
+	GetValueArray() const
+	{ return Values; }
+
+	/** Calculates Min/Max values based on values. */
+	void 
+	CalcRanges(float& MinValue, float& MaxValue)
+	{
+		MinValue = TNumericLimits<float>::Max();
+		MaxValue = -TNumericLimits<float>::Max();
+		const float* ValuesPtr = Values.GetData();
+		for (int32 i = 0; i < Values.Num(); ++i)
+		{
+			MinValue = ValuesPtr[i] < MinValue ? ValuesPtr[i] : MinValue;
+			MaxValue = ValuesPtr[i] > MaxValue ? ValuesPtr[i] : MaxValue;
+		}
+	}
+
+	/** Name of the mask, mainly for users to differentiate */
+	UPROPERTY()
+	FName Name;
+
+	/** The currently targeted parameter for the mask */
+	UPROPERTY()
+	//MaskTarget_PhysMesh CurrentTarget;
+	uint8 CurrentTarget;
+
+	/** The actual values stored in the mask */
+	UPROPERTY()
+	TArray<float> Values;
+
+	/** Whether this mask is enabled and able to effect final mesh values */
+	UPROPERTY()
+	bool bEnabled;
+};
diff --git a/Engine/Source/Runtime/ClothingSystemRuntimeNv/ClothingSystemRuntimeNv.Build.cs b/Engine/Source/Runtime/ClothingSystemRuntimeNv/ClothingSystemRuntimeNv.Build.cs
--- a/Engine/Source/Runtime/ClothingSystemRuntimeNv/ClothingSystemRuntimeNv.Build.cs
+++ b/Engine/Source/Runtime/ClothingSystemRuntimeNv/ClothingSystemRuntimeNv.Build.cs
@@ -16,7 +16,8 @@ public class ClothingSystemRuntimeNv : ModuleRules
             new string[] {
                 "Core",
                 "CoreUObject",
-                "ClothingSystemRuntimeInterface"
+                "ClothingSystemRuntimeInterface",
+                "ClothingSystemRuntimeCommon"
             }
         );
 
diff --git a/Engine/Source/Runtime/ClothingSystemRuntimeNv/Private/Assets/ClothingAssetNv.cpp b/Engine/Source/Runtime/ClothingSystemRuntimeNv/Private/Assets/ClothingAssetNv.cpp
new file mode 100644
--- /dev/null
+++ b/Engine/Source/Runtime/ClothingSystemRuntimeNv/Private/Assets/ClothingAssetNv.cpp
@@ -0,0 +1,318 @@
+// Copyright 1998-2019 Epic Games, Inc. All Rights Reserved.
+
+#include "Assets/ClothingAssetNv.h"
+
+#include "PhysXPublic.h"
+
+#if WITH_EDITOR
+#include "Engine/SkeletalMesh.h"
+#include "MeshUtilities.h"
+#endif
+
+#include "Components/SkeletalMeshComponent.h"
+
+#include "PhysicsPublic.h"
+#include "Modules/ModuleManager.h"
+#include "Widgets/Notifications/SNotificationList.h"
+#include "Framework/Notifications/NotificationManager.h"
+#include "UObject/AnimPhysObjectVersion.h"
+#include "Rendering/SkeletalMeshModel.h"
+#include "Rendering/SkeletalMeshRenderData.h"
+#include "ClothingSimulationInteractor.h"
+#include "Serialization/CustomVersion.h"
+
+#define LOCTEXT_NAMESPACE "ClothingAsset"
+
+//==============================================================================
+// FClothingAssetCustomVersion
+//==============================================================================
+
+// Custom serialization version for clothing assets
+struct FClothingAssetCustomVersion
+{
+	enum Type
+	{
+		// Before any version changes were made
+		BeforeCustomVersionWasAdded = 0,
+		// Added storage of vertex colors with sim data, for editor usage
+		AddVertexColorsToPhysicalMesh = 1,
+		// Changed memory layouts by moving some properties to base classes
+		MovePropertiesToCommonBaseClasses = 2,
+
+		// -----<new versions can be added above this line>-------------------------------------------------
+		VersionPlusOne,
+		LatestVersion = VersionPlusOne - 1
+	};
+
+	// The GUID for this custom version number
+	const static FGuid GUID;
+
+private:
+	FClothingAssetCustomVersion() {}
+};
+
+const FGuid FClothingAssetCustomVersion::GUID(0xFB680AF2, 0x59EF4BA3, 0xBAA819B5, 0x73C8443D);
+FCustomVersionRegistration GRegisterClothingAssetCustomVersion(FClothingAssetCustomVersion::GUID, FClothingAssetCustomVersion::LatestVersion, TEXT("ClothingAssetVer"));
+
+//==============================================================================
+// UClothingAssetNv
+//==============================================================================
+
+UClothingAssetNv::UClothingAssetNv(const FObjectInitializer& ObjectInitializer)
+	: Super(ObjectInitializer)
+{
+	ClothSimConfig = ObjectInitializer.CreateDefaultSubobject<UClothConfigNv>(this, UClothConfigNv::StaticClass()->GetFName());
+}
+
+void UClothingAssetNv::PostLoad()
+{
+	Super::PostLoad();
+
+	const int32 AnimPhysCustomVersion = GetLinkerCustomVersion(FAnimPhysObjectVersion::GUID);
+	const int32 ClothingCustomVersion = GetLinkerCustomVersion(FClothingAssetCustomVersion::GUID);
+
+	if (ClothingCustomVersion < FClothingAssetCustomVersion::MovePropertiesToCommonBaseClasses)
+	{
+		// Remap struct FClothConfig to class UClothConfigNv
+		ClothConfig_DEPRECATED.MigrateTo(Cast<UClothConfigNv>(ClothSimConfig));
+		// Remap struct FClothLODData to class UClothLODDataNv
+		for (FClothLODData &FLod : LodData_DEPRECATED)
+		{
+			const int32 Idx = AddNewLod();
+			FLod.MigrateTo(Cast<UClothLODDataNv>(ClothLodData[Idx]));
+		}
+		LodData_DEPRECATED.Empty();
+	}
+	if(AnimPhysCustomVersion < FAnimPhysObjectVersion::AddedClothingMaskWorkflow)
+	{
+#if WITH_EDITORONLY_DATA
+		// Convert current parameters to masks
+		for(UClothLODDataBase* LodPtr : ClothLodData)
+		{
+			check(LodPtr);
+			UClothLODDataNv& Lod = *Cast<UClothLODDataNv>(LodPtr);
+			check(Lod.PhysicalMeshData);
+			UClothPhysicalMeshDataBase& PhysMesh = *Lod.PhysicalMeshData;
+
+			// Didn't do anything previously - clear out incase there's something in there
+			// so we can use it correctly now.
+			Lod.ParameterMasks.Reset(3);
+
+			// Max distances (Always present)
+			Lod.ParameterMasks.AddDefaulted();
+			FPointWeightMap& MaxDistanceMask = Lod.ParameterMasks.Last();
+			MaxDistanceMask.CopyFrom(PhysMesh.GetFloatArray((uint32)MaskTarget_PhysMesh::MaxDistance), (uint8)MaskTarget_PhysMesh::MaxDistance);
+			MaxDistanceMask.bEnabled = true;
+
+			// Following params are only added if necessary, if we don't have any backstop
+			// radii then there's no backstops.
+			const TArray<float>* BackstopRadiuses = PhysMesh.GetFloatArray((uint32)MaskTarget_PhysMesh::BackstopRadius);
+			if(BackstopRadiuses && BackstopRadiuses->FindByPredicate([](const float& A) {return A != 0.0f; }))
+			{
+				// Backstop radii
+				Lod.ParameterMasks.AddDefaulted();
+				FPointWeightMap& BackstopRadiusMask = Lod.ParameterMasks.Last();
+				BackstopRadiusMask.CopyFrom(PhysMesh.GetFloatArray((uint32)MaskTarget_PhysMesh::BackstopRadius), (uint8)MaskTarget_PhysMesh::BackstopRadius);
+				BackstopRadiusMask.bEnabled = true;
+
+				// Backstop distances
+				Lod.ParameterMasks.AddDefaulted();
+				FPointWeightMap& BackstopDistanceMask = Lod.ParameterMasks.Last();
+				BackstopDistanceMask.CopyFrom(PhysMesh.GetFloatArray((uint32)MaskTarget_PhysMesh::BackstopDistance), (uint8)MaskTarget_PhysMesh::BackstopDistance);
+				BackstopDistanceMask.bEnabled = true;
+			}
+			
+		}
+#endif
+
+		// Make sure we're transactional
+		SetFlags(RF_Transactional);
+	}
+
+#if WITH_EDITORONLY_DATA
+	// Fix content imported before we kept vertex colors
+	if(ClothingCustomVersion < FClothingAssetCustomVersion::AddVertexColorsToPhysicalMesh)
+	{
+		for (UClothLODDataBase* Lod : ClothLodData)
+		{
+			const int32 NumVerts = Lod->PhysicalMeshData->Vertices.Num(); // number of verts
+
+			Lod->PhysicalMeshData->VertexColors.Reset();
+			Lod->PhysicalMeshData->VertexColors.AddUninitialized(NumVerts);
+			for (int32 VertIdx = 0; VertIdx < NumVerts; VertIdx++)
+			{
+				Lod->PhysicalMeshData->VertexColors[VertIdx] = FColor::White;
+			}
+		}
+	}
+#endif // WITH_EDITORONLY_DATA
+
+#if WITH_EDITOR
+	if(AnimPhysCustomVersion < FAnimPhysObjectVersion::CacheClothMeshInfluences)
+	{
+		// Rebuild data cache
+		InvalidateCachedData();
+	}
+#endif
+
+	// After fixing the content, we are ready to call functions that rely on it
+	BuildSelfCollisionData();
+#if WITH_EDITORONLY_DATA
+	CalculateReferenceBoneIndex();
+#endif
+}
+
+void UClothingAssetNv::Serialize(FArchive& Ar)
+{
+	Super::Serialize(Ar);
+	Ar.UsingCustomVersion(FAnimPhysObjectVersion::GUID);
+	Ar.UsingCustomVersion(FClothingAssetCustomVersion::GUID);
+}
+
+#if WITH_EDITOR
+
+void UClothingAssetNv::InvalidateCachedData()
+{
+	for(UClothLODDataBase* CurrentLodDataPtr : ClothLodData)
+	{
+		check(CurrentLodDataPtr);
+		UClothLODDataNv& CurrentLodData = *Cast<UClothLODDataNv>(CurrentLodDataPtr);
+		// Recalculate inverse masses for the physical mesh particles
+		check(CurrentLodData.PhysicalMeshData);
+		UClothPhysicalMeshDataBase& PhysMesh = *CurrentLodData.PhysicalMeshData;
+		check(PhysMesh.Indices.Num() % 3 == 0);
+
+		TArray<float>& InvMasses = PhysMesh.InverseMasses;
+
+		const int32 NumVerts = PhysMesh.Vertices.Num();
+		InvMasses.Empty(NumVerts);
+		InvMasses.AddZeroed(NumVerts);
+
+		for(int32 TriBaseIndex = 0; TriBaseIndex < PhysMesh.Indices.Num(); TriBaseIndex += 3)
+		{
+			const int32 Index0 = PhysMesh.Indices[TriBaseIndex];
+			const int32 Index1 = PhysMesh.Indices[TriBaseIndex + 1];
+			const int32 Index2 = PhysMesh.Indices[TriBaseIndex + 2];
+
+			const FVector AB = PhysMesh.Vertices[Index1] - PhysMesh.Vertices[Index0];
+			const FVector AC = PhysMesh.Vertices[Index2] - PhysMesh.Vertices[Index0];
+			const float TriArea = FVector::CrossProduct(AB, AC).Size();
+
+			InvMasses[Index0] += TriArea;
+			InvMasses[Index1] += TriArea;
+			InvMasses[Index2] += TriArea;
+		}
+
+		PhysMesh.NumFixedVerts = 0;
+
+		const TArray<float>* const MaxDistances = PhysMesh.GetFloatArray((uint32)MaskTarget_PhysMesh::MaxDistance);
+		if (MaxDistances && MaxDistances->Num() > 0)
+		{
+			float MassSum = 0.0f;
+			for (int32 CurrVertIndex = 0; CurrVertIndex < NumVerts; ++CurrVertIndex)
+			{
+				float& InvMass = InvMasses[CurrVertIndex];
+				const float& MaxDistance = (*MaxDistances)[CurrVertIndex];
+
+				if (MaxDistance < SMALL_NUMBER)   // For consistency, the default Threshold should be 0.1, not SMALL_NUMBER. But for backward compatibility it needs to be SMALL_NUMBER for now.
+				{
+					InvMass = 0.0f;
+					++PhysMesh.NumFixedVerts;
+				}
+				else
+				{
+					MassSum += InvMass;
+				}
+			}
+
+			if (MassSum > 0.0f)
+			{
+				const float MassScale = (float)(NumVerts - PhysMesh.NumFixedVerts) / MassSum;
+				for (float& InvMass : InvMasses)
+				{
+					if (InvMass != 0.0f)
+					{
+						InvMass *= MassScale;
+						InvMass = 1.0f / InvMass;
+					}
+				}
+			}
+		}
+		else
+		{
+			// Otherwise, go fully kinematic.
+			for(int32 CurrVertIndex = 0; CurrVertIndex < NumVerts; ++CurrVertIndex)
+			{
+				InvMasses[CurrVertIndex] = 0.0f;
+			}
+			PhysMesh.NumFixedVerts = NumVerts;
+		}
+
+		// Calculate number of influences per vertex
+		for(int32 VertIndex = 0; VertIndex < NumVerts; ++VertIndex)
+		{
+			FClothVertBoneData& BoneData = PhysMesh.BoneData[VertIndex];
+			const uint16* BoneIndices = BoneData.BoneIndices;
+			const float* BoneWeights = BoneData.BoneWeights;
+
+			BoneData.NumInfluences = MAX_TOTAL_INFLUENCES;
+
+			int32 NumInfluences = 0;
+			for(int32 InfluenceIndex = 0; InfluenceIndex < MAX_TOTAL_INFLUENCES; ++InfluenceIndex)
+			{
+				if(BoneWeights[InfluenceIndex] == 0.0f || BoneIndices[InfluenceIndex] == INDEX_NONE)
+				{
+					BoneData.NumInfluences = NumInfluences;
+					break;
+				}
+				++NumInfluences;
+			}
+		}
+	}
+}
+
+int32 UClothingAssetNv::AddNewLod()
+{
+	const int32 Idx = ClothLodData.AddDefaulted();
+	ClothLodData[Idx] = NewObject<UClothLODDataNv>(this, UClothLODDataNv::StaticClass(), NAME_None, RF_Transactional);
+	return Idx;
+}
+
+void UClothingAssetNv::PostPropertyChangeCb(const FPropertyChangedEvent& Event)
+{
+	bool bReregisterComponents = false;
+
+	if(Event.ChangeType != EPropertyChangeType::Interactive)
+	{
+		if(Event.Property->GetFName() == GET_MEMBER_NAME_CHECKED(UClothConfigNv, SelfCollisionRadius) ||
+		   Event.Property->GetFName() == GET_MEMBER_NAME_CHECKED(UClothConfigNv, SelfCollisionCullScale))
+		{
+			BuildSelfCollisionData();
+			bReregisterComponents = true;
+		}
+		else if(Event.Property->GetFName() == GET_MEMBER_NAME_CHECKED(UClothingAssetCommon, PhysicsAsset))
+		{
+			bReregisterComponents = true;
+		}
+		else
+		{
+			// Other properties just require a config refresh
+			ForEachInteractorUsingClothing([](UClothingSimulationInteractor* InInteractor)
+			{
+				if(InInteractor)
+				{
+					InInteractor->ClothConfigUpdated();
+				}
+			});
+		}
+	}
+
+	if(bReregisterComponents)
+	{
+		ReregisterComponentsUsingClothing();
+	}
+}
+
+#endif // WITH_EDITOR
+
+#undef LOCTEXT_NAMESPACE
diff --git a/Engine/Source/Runtime/ClothingSystemRuntimeNv/Private/ClothConfigNv.cpp b/Engine/Source/Runtime/ClothingSystemRuntimeNv/Private/ClothConfigNv.cpp
new file mode 100644
--- /dev/null
+++ b/Engine/Source/Runtime/ClothingSystemRuntimeNv/Private/ClothConfigNv.cpp
@@ -0,0 +1,47 @@
+#include "ClothConfigNv.h"
+
+UClothConfigNv::UClothConfigNv()
+	: WindMethod(EClothingWindMethod::Legacy)
+	, SelfCollisionRadius(0.0f)
+	, SelfCollisionStiffness(0.0f)
+	, SelfCollisionCullScale(1.0f)
+	, Damping(0.4f)
+	, Friction(0.1f)
+	, WindDragCoefficient(0.02f/100.0f)
+	, WindLiftCoefficient(0.02f/100.0f)
+	, LinearDrag(0.2f)
+	, AngularDrag(0.2f)
+	, LinearInertiaScale(1.0f)
+	, AngularInertiaScale(1.0f)
+	, CentrifugalInertiaScale(1.0f)
+	, SolverFrequency(120.0f)
+	, StiffnessFrequency(100.0f)
+	, GravityScale(1.0f)
+	, GravityOverride(FVector::ZeroVector)
+	, bUseGravityOverride(false)
+	, TetherStiffness(1.0f)
+	, TetherLimit(1.0f)
+	, CollisionThickness(1.0f)
+	, AnimDriveSpringStiffness(1.0f)
+	, AnimDriveDamperStiffness(1.0f)
+{}
+
+bool UClothConfigNv::HasSelfCollision() const
+{
+	return SelfCollisionRadius > 0.0f && SelfCollisionStiffness > 0.0f;
+}
+
+#if WITH_EDITOR
+bool UClothConfigNv::InitFromApexAssetCallback(
+	nvidia::apex::ClothingAsset* InApexAsset,
+	USkeletalMesh* TargetMesh,
+	FName InName) 
+{
+#if WITH_APEX_CLOTHING
+	// Set to use legacy wind calculations, which is what APEX would normally have used
+	WindMethod = EClothingWindMethod::Legacy;
+#endif // WITH_APEX_CLOTHING
+	return false;
+}
+
+#endif //WITH_EDITOR
diff --git a/Engine/Source/Runtime/ClothingSystemRuntimeNv/Private/ClothLODDataNv.cpp b/Engine/Source/Runtime/ClothingSystemRuntimeNv/Private/ClothLODDataNv.cpp
new file mode 100644
--- /dev/null
+++ b/Engine/Source/Runtime/ClothingSystemRuntimeNv/Private/ClothLODDataNv.cpp
@@ -0,0 +1,11 @@
+// Copyright 1998-2019 Epic Games, Inc. All Rights Reserved.
+#include "ClothLODDataNv.h"
+
+UClothLODDataNv::UClothLODDataNv(const FObjectInitializer& Init)
+	: Super(Init)
+{
+	PhysicalMeshData = Init.CreateDefaultSubobject<UClothPhysicalMeshDataNv>(this, FName("ClothPhysicalMeshDataNv"));
+}
+
+UClothLODDataNv::~UClothLODDataNv()
+{}
diff --git a/Engine/Source/Runtime/ClothingSystemRuntimeNv/Private/ClothPhysicalMeshDataNv.cpp b/Engine/Source/Runtime/ClothingSystemRuntimeNv/Private/ClothPhysicalMeshDataNv.cpp
new file mode 100644
--- /dev/null
+++ b/Engine/Source/Runtime/ClothingSystemRuntimeNv/Private/ClothPhysicalMeshDataNv.cpp
@@ -0,0 +1,106 @@
+// Copyright 1998-2019 Epic Games, Inc. All Rights Reserved.
+#include "ClothPhysicalMeshDataNv.h"
+#include "Assets/ClothingAssetNv.h" // For UClothConfigNv
+
+UClothPhysicalMeshDataNv::UClothPhysicalMeshDataNv()
+{
+	Super::RegisterFloatArray((uint32)MaskTarget_PhysMesh::MaxDistance, &MaxDistances);
+	Super::RegisterFloatArray((uint32)MaskTarget_PhysMesh::BackstopDistance, &BackstopDistances);
+	Super::RegisterFloatArray((uint32)MaskTarget_PhysMesh::BackstopRadius, &BackstopRadiuses);
+	Super::RegisterFloatArray((uint32)MaskTarget_PhysMesh::AnimDriveMultiplier, &AnimDriveMultipliers);
+}
+
+UClothPhysicalMeshDataNv::~UClothPhysicalMeshDataNv()
+{}
+
+void UClothPhysicalMeshDataNv::Reset(const int32 NumPoints)
+{
+	Super::Reset(NumPoints);
+
+	MaxDistances.Reset();
+	BackstopDistances.Reset();
+	BackstopRadiuses.Reset();
+	AnimDriveMultipliers.Reset();
+
+	MaxDistances.AddDefaulted(NumPoints);
+	BackstopDistances.AddDefaulted(NumPoints);
+	BackstopRadiuses.AddDefaulted(NumPoints);
+	AnimDriveMultipliers.AddDefaulted(NumPoints);
+}
+
+void UClothPhysicalMeshDataNv::ClearParticleParameters()
+{
+	Super::ClearParticleParameters();
+
+	// Max distances must be present, so fill to zero on clear so we still have valid mesh data.
+	const int32 NumVerts = Vertices.Num();
+	MaxDistances.Reset(NumVerts);
+	MaxDistances.AddZeroed(NumVerts);
+
+	// Just clear optional properties
+	BackstopDistances.Empty();
+	BackstopRadiuses.Empty();
+	AnimDriveMultipliers.Empty();
+}
+
+void UClothPhysicalMeshDataNv::BuildSelfCollisionData(const UClothConfigBase* ClothConfigBase)
+{
+	const UClothConfigNv* ClothConfig = Cast<UClothConfigNv>(ClothConfigBase);
+	const float SCRadius = ClothConfig ? ClothConfig->SelfCollisionRadius * ClothConfig->SelfCollisionCullScale : 0.1f;
+	const float SCRadiusSq = SCRadius * SCRadius;
+
+	// Start with the full set
+	const int32 NumVerts = Vertices.Num();
+	SelfCollisionIndices.Reset();
+	for(int32 Index = 0; Index < NumVerts; ++Index)
+	{
+		if(!IsFixed(Index))
+			SelfCollisionIndices.Add(Index);
+	}
+
+	// Now start aggresively culling verts that are near others that we have accepted
+	for (int32 Vert0Itr = 0; Vert0Itr < SelfCollisionIndices.Num(); ++Vert0Itr)
+	{
+		const uint32 V0Index = SelfCollisionIndices[Vert0Itr];
+		if (V0Index == INDEX_NONE)
+		{
+			// We'll remove these indices later.  Just skip it for now.
+			continue;
+		}
+
+		const FVector& V0Pos = Vertices[V0Index];
+
+		// Start one after our current V0, we've done the other checks
+		for (int32 Vert1Itr = Vert0Itr + 1; Vert1Itr < SelfCollisionIndices.Num(); ++Vert1Itr)
+		{
+			const uint32 V1Index = SelfCollisionIndices[Vert1Itr];
+			if (V1Index == INDEX_NONE)
+			{
+				// We'll remove these indices later.  Just skip it for now.
+				continue;
+			}
+
+			const FVector& V1Pos = Vertices[V1Index];
+			const float V0ToV1DistSq = (V1Pos - V0Pos).SizeSquared();
+			if (V0ToV1DistSq < SCRadiusSq)
+			{
+				// Points are in contact in the rest state.  Remove it.
+				//
+				// It's worth noting that this biases towards removing indices 
+				// of later in the list, and keeping ones earlier.  That's not
+				// a great criteria for choosing which one is more important.
+				SelfCollisionIndices[Vert1Itr] = INDEX_NONE;
+				continue;
+			}
+		}
+	}
+
+	// Cull flagged indices.
+	for (int32 It = SelfCollisionIndices.Num(); It--;)
+	{
+		if (SelfCollisionIndices[It] == INDEX_NONE)
+		{
+			SelfCollisionIndices.RemoveAt(It);
+		}
+	}
+}
diff --git a/Engine/Source/Runtime/ClothingSystemRuntimeNv/Private/ClothingSimulationNv.cpp b/Engine/Source/Runtime/ClothingSystemRuntimeNv/Private/ClothingSimulationNv.cpp
--- a/Engine/Source/Runtime/ClothingSystemRuntimeNv/Private/ClothingSimulationNv.cpp
+++ b/Engine/Source/Runtime/ClothingSystemRuntimeNv/Private/ClothingSimulationNv.cpp
@@ -4,7 +4,7 @@
 
 #if WITH_NVCLOTH
 
-#include "Assets/ClothingAsset.h"
+#include "Assets/ClothingAssetNv.h"
 #include "ClothingSystemRuntimeModuleNv.h"
 #include "NvClothIncludes.h"
 #include "NvClothSupport.h"
@@ -52,8 +52,7 @@ FClothingSimulationBase::FClothingSimulationBase()
 FClothingSimulationBase::~FClothingSimulationBase()
 {}
 
-void 
-FClothingSimulationBase::FillContext(
+void FClothingSimulationBase::FillContext(
 	USkeletalMeshComponent* InComponent, 
 	float InDeltaTime, 
 	IClothingSimulationContext* InOutContext)
@@ -188,7 +187,7 @@ void FClothingSimulationNv::CreateActor(USkeletalMeshComponent* InOwnerComponent
 	check(CachedFactory);
 	check(InOwnerComponent);
 
-	UClothingAsset* Asset = Cast<UClothingAsset>(InAsset);
+	UClothingAssetCommon* Asset = Cast<UClothingAssetCommon>(InAsset);
 	
 	UWorld* World = InOwnerComponent->GetWorld();
 	FPhysScene* PhysicsScene = World ? World->GetPhysicsScene() : nullptr;
@@ -207,17 +206,22 @@ void FClothingSimulationNv::CreateActor(USkeletalMeshComponent* InOwnerComponent
 
 	// Need the current reftolocals so we can skin the ref pose for the sim mesh
 	TArray<FMatrix> RefToLocals;
-	InOwnerComponent->GetCurrentRefToLocalMatrices(RefToLocals, FMath::Min(InOwnerComponent->PredictedLODLevel, Asset->LodData.Num() - 1));
+	InOwnerComponent->GetCurrentRefToLocalMatrices(RefToLocals, FMath::Min(InOwnerComponent->PredictedLODLevel, Asset->ClothLodData.Num() - 1));
 
 	Actors.AddDefaulted();
 	FClothingActorNv& NewActor = Actors.Last();
 
 	NewActor.SimDataIndex = InSimDataIndex;
 
-	for(int32 LodIndex = 0; LodIndex < Asset->LodData.Num(); ++LodIndex)
+	for(int32 LodIndex = 0; LodIndex < Asset->ClothLodData.Num(); ++LodIndex)
 	{
-		FClothLODData& AssetLodData = Asset->LodData[LodIndex];
-		FClothPhysicalMeshData& PhysMesh = AssetLodData.PhysicalMeshData;
+		UClothLODDataBase& AssetLodData = *Asset->ClothLodData[LodIndex];
+		UClothPhysicalMeshDataNv* PhysMeshPtr = Cast<UClothPhysicalMeshDataNv>(AssetLodData.PhysicalMeshData);
+		if (!ensure(PhysMeshPtr))
+		{
+			continue;
+		}
+		UClothPhysicalMeshDataNv& PhysMesh = *PhysMeshPtr;
 
 		NewActor.LodData.AddDefaulted();
 		FClothingActorNv::FActorLodData& ActorLodData = NewActor.LodData.Last();
@@ -360,7 +364,7 @@ void FClothingSimulationNv::CreateActor(USkeletalMeshComponent* InOwnerComponent
 		}
 	}
 
-	ApplyClothConfig(Asset->ClothConfig, NewActor);
+	ApplyClothConfig(Asset->ClothSimConfig, NewActor);
 
 	// Pull collisions from the specified physics asset inside the clothing asset
 	ExtractActorCollisions(Asset, NewActor);
@@ -382,7 +386,7 @@ void FClothingSimulationNv::CreateActor(USkeletalMeshComponent* InOwnerComponent
 	}
 }
 
-void FClothingSimulationNv::ExtractActorCollisions(UClothingAsset* Asset, FClothingActorNv &InActor)
+void FClothingSimulationNv::ExtractActorCollisions(UClothingAssetCommon* Asset, FClothingActorNv &InActor)
 {
 	// Build collision data for this Actor
 	USkeletalMesh* TargetMesh = CastChecked<USkeletalMesh>(Asset->GetOuter());
@@ -470,8 +474,13 @@ void FClothingSimulationNv::ExtractActorCollisions(UClothingAsset* Asset, FCloth
 	}
 }
 
-void FClothingSimulationNv::ApplyClothConfig(FClothConfig &Config, FClothingActorNv &InActor)
+void FClothingSimulationNv::ApplyClothConfig(UClothConfigBase *BaseConfig, FClothingActorNv &InActor)
 {
+	check(BaseConfig);
+	UClothConfigNv* ConfigPtr = Cast<UClothConfigNv>(BaseConfig);
+	check(ConfigPtr);
+	UClothConfigNv& Config = *ConfigPtr;
+
 	// These calculations convert our damping and drag values to a value closer to the way APEX used them.
 	// This allows us to not break old assets, and expectations from the old system
 	const float DampStiffnesssFreq = 10.0f;
@@ -669,7 +678,7 @@ void FClothingSimulationNv::Simulate(IClothingSimulationContext* InContext)
 
 		// To build motion constraints (max distances) we need to skin the entire physics mesh
 		// this call also updates our fixed particles to avoid iterating the particle list a second time
-		const FClothPhysicalMeshData& PhysMesh = Actor.AssetCreatedFrom->LodData[Actor.CurrentLodIndex].PhysicalMeshData;
+		const UClothPhysicalMeshDataBase& PhysMesh = *Actor.AssetCreatedFrom->ClothLodData[Actor.CurrentLodIndex]->PhysicalMeshData;
 		ClothingMeshUtils::SkinPhysicsMesh(Actor.AssetCreatedFrom->UsedBoneIndices, PhysMesh, RootBoneTransform, NvContext->RefToLocals.GetData(), NvContext->RefToLocals.Num(), Actor.SkinnedPhysicsMeshPositions[Actor.CurrentSkinnedPositionIndex], Actor.SkinnedPhysicsMeshNormals);
 
 		nv::cloth::Cloth* CurrentCloth = Actor.LodData[Actor.CurrentLodIndex].Cloth;
@@ -708,7 +717,11 @@ void FClothingSimulationNv::Simulate(IClothingSimulationContext* InContext)
 		CurrentCloth->setTranslation(U2PVector(RootBoneWorldTransform.GetTranslation()));
 		CurrentCloth->setRotation(U2PQuat(RootBoneWorldTransform.GetRotation()));
 
-		FClothConfig& ClothConfig = Actor.AssetCreatedFrom->ClothConfig;
+		UClothConfigBase* ClothConfigBase = Actor.AssetCreatedFrom->ClothSimConfig;
+		check(ClothConfigBase);
+		UClothConfigNv* ConfigPtr = Cast<UClothConfigNv>(ClothConfigBase);
+		check(ConfigPtr);
+		UClothConfigNv& ClothConfig = *ConfigPtr;
 		if (ClothConfig.bUseGravityOverride)
 		{
 			CurrentCloth->setGravity(U2PVector(ClothConfig.GravityOverride));
@@ -854,7 +867,7 @@ void FClothingSimulationNv::ComputePhysicalMeshNormals(FClothingActorNv &Actor)
 
 	const uint32 NumParticles = Actor.LodData[CurrentClothingLod].Cloth->getNumParticles();
 	nv::cloth::MappedRange<physx::PxVec4> Particles = Actor.LodData[CurrentClothingLod].Cloth->getCurrentParticles();
-	const TArray<uint32>& Indices = Actor.AssetCreatedFrom->LodData[CurrentClothingLod].PhysicalMeshData.Indices;
+	const TArray<uint32>& Indices = Actor.AssetCreatedFrom->ClothLodData[CurrentClothingLod]->PhysicalMeshData->Indices;
 	const uint32 NumIndices = Indices.Num();
 
 	// Using the face normals, calculate normals. These will not be normalized as we're adding together
@@ -932,7 +945,7 @@ void FClothingSimulationNv::GetSimulationData(TMap<int32, FClothSimulData>& OutD
 	for(int32 ActorIdx = 0; ActorIdx < NumActors; ++ActorIdx)
 	{
 		const FClothingActorNv& Actor = Actors[ActorIdx];
-		UClothingAsset* Asset = Actor.AssetCreatedFrom;
+		UClothingAssetCommon* Asset = Actor.AssetCreatedFrom;
 
 		if(Actor.CurrentLodIndex == INDEX_NONE)
 		{
@@ -1012,7 +1025,7 @@ void FClothingSimulationNv::GetCollisions(FClothCollisionData& OutCollisions, bo
 		else
 		{
 			check(Actor.AssetCreatedFrom);
-			OutCollisions.Append(Actor.AssetCreatedFrom->LodData[Actor.CurrentLodIndex].CollisionData);
+			OutCollisions.Append(Actor.AssetCreatedFrom->ClothLodData[Actor.CurrentLodIndex]->CollisionData);
 		}
 	}
 }
@@ -1029,7 +1042,7 @@ void FClothingSimulationNv::GatherStats() const
 			continue;
 		}
 
-		INC_DWORD_STAT_BY(STAT_NumClothVerts, Actor.AssetCreatedFrom->LodData[Actor.CurrentLodIndex].PhysicalMeshData.Vertices.Num());
+		INC_DWORD_STAT_BY(STAT_NumClothVerts, Actor.AssetCreatedFrom->ClothLodData[Actor.CurrentLodIndex]->PhysicalMeshData->Vertices.Num());
 	}
 }
 
@@ -1140,23 +1153,9 @@ void FClothingSimulationNv::UpdateLod(int32 InPredictedLod, const FTransform& Co
 			bool bOldLodMapped = LodMap.IsValidIndex(CurrentMeshLodIndex) && LodMap[CurrentMeshLodIndex] != INDEX_NONE;
 
 			// Get the clothing LOD mapped from the mesh predicted LOD
-
-
+			const int32 PredictedClothingLod = LodMap[InPredictedLod];
 			const int32 OldClothingLod = bOldLodMapped ? LodMap[CurrentMeshLodIndex] : INDEX_NONE;
 
-			// If potentialLod doesn't map to a valid LOD, we try higher LOD levels for a valid LOD.
-			// Asset might only have lod on LOD 1 and not 0, however if mesh doesn't force LOD to 1, 
-			// asset will not be assigned valid LOD index and will not generate sim data, breaking things.
-			int32 PredictedClothingLod = INDEX_NONE;
-			for (int32 PotentialLod = InPredictedLod; PotentialLod < LodMap.Num(); ++PotentialLod)
-			{
-				if (LodMap[PotentialLod] != INDEX_NONE)
-				{
-					PredictedClothingLod = LodMap[PotentialLod];
-					break;
-				}
-			}
-
 			if(PredictedClothingLod == Actor.CurrentLodIndex)
 			{
 				// We must have forced a LOD update because we added a new actor - this actor is good to go though.
@@ -1208,7 +1207,7 @@ void FClothingSimulationNv::UpdateLod(int32 InPredictedLod, const FTransform& Co
 				{
 					// Reposition particles skinned to outgoing LOD
 					bool bLodTransitionUp = OldClothingLod < PredictedClothingLod;
-					FClothLODData& NewLodAssetData = Actor.AssetCreatedFrom->LodData[PredictedClothingLod];
+					UClothLODDataBase& NewLodAssetData = *Actor.AssetCreatedFrom->ClothLodData[PredictedClothingLod];
 					TArray<FMeshToMeshVertData>& SkinData = bLodTransitionUp ? NewLodAssetData.TransitionUpSkinData : NewLodAssetData.TransitionDownSkinData;
 
 					for(int32 ParticleIndex = 0; ParticleIndex < NumNewParticles; ++ParticleIndex)
@@ -1279,7 +1278,7 @@ void FClothingSimulationNv::UpdateLod(int32 InPredictedLod, const FTransform& Co
 					if(CSTransforms.Num() > Actor.AssetCreatedFrom->ReferenceBoneIndex)
 					{
 						// compute skinned positions to init sim mesh
-						const FClothPhysicalMeshData& PhysMesh = Actor.AssetCreatedFrom->LodData[PredictedClothingLod].PhysicalMeshData;
+						const UClothPhysicalMeshDataBase& PhysMesh = *Actor.AssetCreatedFrom->ClothLodData[PredictedClothingLod]->PhysicalMeshData;
 						TArray<FVector> SkinnedPhysicsMeshPositions;
 						TArray<FVector> SkinnedPhysicsMeshNormals;
 						ClothingMeshUtils::SkinPhysicsMesh(
@@ -1369,8 +1368,14 @@ void FClothingSimulationNv::DebugDraw_PhysMesh(USkeletalMeshComponent* OwnerComp
 			Particles = NvClothSupport::CreateRange(ActorData.Px_RestPositions);
 		}
 
-		TArray<uint32>& Indices = Actor.AssetCreatedFrom->LodData[CurrentClothLod].PhysicalMeshData.Indices;
-		TArray<float>& MaxDistances = Actor.AssetCreatedFrom->LodData[CurrentClothLod].PhysicalMeshData.MaxDistances;
+		UClothPhysicalMeshDataBase* PhysMeshPtr = Actor.AssetCreatedFrom->ClothLodData[CurrentClothLod]->PhysicalMeshData;
+		TArray<uint32>& Indices = PhysMeshPtr->Indices;
+		UClothPhysicalMeshDataNv* PhysMeshNvPtr = Cast<UClothPhysicalMeshDataNv>(PhysMeshPtr);
+		if (!ensure(PhysMeshNvPtr))
+		{
+			continue;
+		}
+		TArray<float>& MaxDistances = PhysMeshNvPtr->MaxDistances;
 
 		const int32 NumTriangles = Indices.Num() / 3;
 
@@ -1516,28 +1521,44 @@ void FClothingSimulationNv::DebugDraw_Backstops(USkeletalMeshComponent* OwnerCom
 		}
 
 		const FClothingActorNv::FActorLodData& LodData = Actor.LodData[Actor.CurrentLodIndex];
-		const UClothingAsset* Asset = Actor.AssetCreatedFrom;
-		const FClothPhysicalMeshData& MeshData = Asset->LodData[Actor.CurrentLodIndex].PhysicalMeshData;
-		FTransform RootBoneTransform = OwnerComponent->GetComponentSpaceTransforms()[Actor.AssetCreatedFrom->ReferenceBoneIndex];
-
-		if(!MeshData.HasBackStops())
+		const UClothingAssetCommon* Asset = Actor.AssetCreatedFrom;
+		if (!Asset)
+		{
+			continue;
+		}
+		const UClothPhysicalMeshDataBase* MeshDataBase = Asset->ClothLodData[Actor.CurrentLodIndex]->PhysicalMeshData;
+		if (!MeshDataBase)
+		{
+			continue;
+		}
+		const TArray<float>* BackstopDistancesPtr = MeshDataBase->GetFloatArray((int32)MaskTarget_PhysMesh::BackstopDistance);
+		const TArray<float>* BackstopRadiusesPtr = MeshDataBase->GetFloatArray((int32)MaskTarget_PhysMesh::BackstopRadius);
+		const TArray<float>* MaxDistancesPtr = MeshDataBase->GetFloatArray((int32)MaskTarget_PhysMesh::MaxDistance);
+		if (!BackstopDistancesPtr || !BackstopDistancesPtr->Num() || 
+			!BackstopRadiusesPtr || !BackstopRadiusesPtr->Num() ||
+			!MaxDistancesPtr || !MaxDistancesPtr->Num())
 		{
 			continue;
 		}
+		const TArray<float>& BackstopDistances = *BackstopDistancesPtr;
+		const TArray<float>& BackstopRadiuses = *BackstopRadiusesPtr;
+		const TArray<float>& MaxDistances = *MaxDistancesPtr;
+
+		FTransform RootBoneTransform = OwnerComponent->GetComponentSpaceTransforms()[Actor.AssetCreatedFrom->ReferenceBoneIndex];
 
 		const TArray<FVector>& SkinnedPositions = Actor.GetCurrentSkinnedPositions();
 		const int32 NumVerts = SkinnedPositions.Num();
 		check(NumVerts == Actor.SkinnedPhysicsMeshNormals.Num());
-		check(NumVerts == MeshData.MaxDistances.Num());
+		check(NumVerts == MaxDistances.Num());
 
 		for(int32 VertIndex = 0; VertIndex < NumVerts; ++VertIndex)
 		{
 			const FVector& Position = RootBoneTransform.TransformPosition(SkinnedPositions[VertIndex]);
 			const FVector& Normal = RootBoneTransform.TransformVector(Actor.SkinnedPhysicsMeshNormals[VertIndex]);
 			
-			float BackstopDistance = MeshData.BackstopDistances[VertIndex];
-			const float BackstopRadius = MeshData.BackstopRadiuses[VertIndex];
-			const float MaxDistance = MeshData.MaxDistances[VertIndex];
+			float BackstopDistance = BackstopDistances[VertIndex];
+			const float BackstopRadius = BackstopRadiuses[VertIndex];
+			const float MaxDistance = MaxDistances[VertIndex];
 
 			FColor FixedPointColor = FColor::White;
 
@@ -1581,20 +1602,35 @@ void FClothingSimulationNv::DebugDraw_MaxDistances(USkeletalMeshComponent* Owner
 		}
 
 		const FClothingActorNv::FActorLodData& LodData = Actor.LodData[Actor.CurrentLodIndex];
-		const UClothingAsset* Asset = Actor.AssetCreatedFrom;
-		const FClothPhysicalMeshData& MeshData = Asset->LodData[Actor.CurrentLodIndex].PhysicalMeshData;
+		const UClothingAssetCommon* Asset = Actor.AssetCreatedFrom;
+		if (!Asset)
+		{
+			continue;
+		}
+		const UClothPhysicalMeshDataBase* MeshDataBase = Asset->ClothLodData[Actor.CurrentLodIndex]->PhysicalMeshData;
+		if (!MeshDataBase)
+		{
+			continue;
+		}
+		const TArray<float>* MaxDistancesPtr = MeshDataBase ->GetFloatArray((int32)MaskTarget_PhysMesh::MaxDistance);
+		if (!MaxDistancesPtr || !MaxDistancesPtr->Num())
+		{
+			continue;
+		}
+		const TArray<float>& MaxDistances = *MaxDistancesPtr;
+
 		FTransform RootBoneTransform = OwnerComponent->GetComponentSpaceTransforms()[Actor.AssetCreatedFrom->ReferenceBoneIndex];
 
 		const TArray<FVector>& SkinnedPositions = Actor.GetCurrentSkinnedPositions();
 		const int32 NumVerts = SkinnedPositions.Num();
 		check(NumVerts == Actor.SkinnedPhysicsMeshNormals.Num());
-		check(NumVerts == MeshData.MaxDistances.Num());
+		check(NumVerts == MaxDistances.Num());
 
 		for(int32 VertIndex = 0; VertIndex < NumVerts; ++VertIndex)
 		{
 			const FVector& Position = RootBoneTransform.TransformPosition(SkinnedPositions[VertIndex]);
 			const FVector& Normal = RootBoneTransform.TransformVector(Actor.SkinnedPhysicsMeshNormals[VertIndex]);
-			const float& MaxDistance = MeshData.MaxDistances[VertIndex];
+			const float& MaxDistance = MaxDistances[VertIndex];
 
 			PDI->DrawLine(Position, Position + Normal * MaxDistance, FColor::White, SDPG_World, 0.0f, 0.001f);
 		}
@@ -1611,8 +1647,12 @@ void FClothingSimulationNv::DebugDraw_SelfCollision(USkeletalMeshComponent* Owne
 			continue;
 		}
 
-		const UClothingAsset* Asset = Actor.AssetCreatedFrom;
-		const FClothConfig& Config = Asset->ClothConfig;
+		const UClothingAssetCommon* Asset = Actor.AssetCreatedFrom;
+		UClothConfigBase* ClothConfigBase = Actor.AssetCreatedFrom->ClothSimConfig;
+		check(ClothConfigBase);
+		UClothConfigNv* ConfigPtr = Cast<UClothConfigNv>(ClothConfigBase);
+		check(ConfigPtr);
+		UClothConfigNv& Config = *ConfigPtr;
 
 		if(!Config.HasSelfCollision())
 		{
@@ -1624,8 +1664,8 @@ void FClothingSimulationNv::DebugDraw_SelfCollision(USkeletalMeshComponent* Owne
 
 		const float SelfCollisionThickness = Config.SelfCollisionRadius;
 
-		const FClothLODData& LodData = Asset->LodData[Actor.CurrentLodIndex];
-		const FClothPhysicalMeshData& PhysMesh = LodData.PhysicalMeshData;
+		const UClothLODDataBase& LodData = *Asset->ClothLodData[Actor.CurrentLodIndex];
+		const UClothPhysicalMeshDataBase& PhysMesh = *LodData.PhysicalMeshData;
 
 		nv::cloth::Cloth* CurrentCloth = Actor.LodData[Actor.CurrentLodIndex].Cloth;
 
@@ -1654,8 +1694,8 @@ void FClothingSimulationNv::DebugDraw_AnimDrive(USkeletalMeshComponent* OwnerCom
 			continue;
 		}
 
-		const UClothingAsset* Asset = Actor.AssetCreatedFrom;
-		const FClothPhysicalMeshData& PhysMesh = Asset->LodData[Actor.CurrentLodIndex].PhysicalMeshData;
+		const UClothingAssetCommon* Asset = Actor.AssetCreatedFrom;
+		const UClothPhysicalMeshDataBase& PhysMesh = *Asset->ClothLodData[Actor.CurrentLodIndex]->PhysicalMeshData;
 		const TArray<uint32>& Indices = PhysMesh.Indices;
 		FTransform RootBoneTransform = OwnerComponent->GetComponentSpaceTransforms()[Actor.AssetCreatedFrom->ReferenceBoneIndex];
 
@@ -1688,7 +1728,7 @@ void FClothingSimulationNv::RefreshClothConfig()
 
 	for(FClothingActorNv& Actor : Actors)
 	{
-		ApplyClothConfig(Actor.AssetCreatedFrom->ClothConfig, Actor);
+		ApplyClothConfig(Actor.AssetCreatedFrom->ClothSimConfig, Actor);
 	}
 }
 
@@ -1721,7 +1761,7 @@ FClothingActorNv::FClothingActorNv()
 
 void FClothingActorNv::SkinPhysicsMesh(FClothingSimulationContextNv* InContext)
 {
-	const FClothPhysicalMeshData& PhysMesh = AssetCreatedFrom->LodData[CurrentLodIndex].PhysicalMeshData;
+	const UClothPhysicalMeshDataBase& PhysMesh = *AssetCreatedFrom->ClothLodData[CurrentLodIndex]->PhysicalMeshData;
 	FTransform RootBoneTransform = InContext->BoneTransforms[AssetCreatedFrom->ReferenceBoneIndex];
 	ClothingMeshUtils::SkinPhysicsMesh(AssetCreatedFrom->UsedBoneIndices, PhysMesh, RootBoneTransform, InContext->RefToLocals.GetData(), InContext->RefToLocals.Num(), SkinnedPhysicsMeshPositions[CurrentSkinnedPositionIndex], SkinnedPhysicsMeshNormals);
 }
@@ -1737,10 +1777,14 @@ void FClothingActorNv::UpdateMotionConstraints(FClothingSimulationContextNv* InC
 
 	check(CurrentCloth);
 
-	FClothPhysicalMeshData& PhysMesh = AssetCreatedFrom->LodData[CurrentLodIndex].PhysicalMeshData;
-	TArray<float>& MaxDistances = PhysMesh.MaxDistances;
-	TArray<float>& BackstopDistances = PhysMesh.BackstopDistances;
-	TArray<float>& BackstopRadiuses = PhysMesh.BackstopRadiuses;
+	UClothPhysicalMeshDataNv* PhysMeshPtr = Cast<UClothPhysicalMeshDataNv>(AssetCreatedFrom->ClothLodData[CurrentLodIndex]->PhysicalMeshData);
+	if (!ensure(PhysMeshPtr))
+	{
+		return;
+	}
+	TArray<float>& MaxDistances = PhysMeshPtr->MaxDistances;
+	TArray<float>& BackstopDistances = PhysMeshPtr->BackstopDistances;
+	TArray<float>& BackstopRadiuses = PhysMeshPtr->BackstopRadiuses;
 	const TArray<FVector>& CurrentSkinnedPositions = GetCurrentSkinnedPositions();
 
 	nv::cloth::Range<physx::PxVec4> MotionConstraints = CurrentCloth->getMotionConstraints();
@@ -1752,11 +1796,12 @@ void FClothingActorNv::UpdateMotionConstraints(FClothingSimulationContextNv* InC
 		MotionConstraints[ConstraintIndex] = physx::PxVec4(U2PVector(CurrentSkinnedPositions[ConstraintIndex]), MaxDistances[ConstraintIndex] * InContext->MaxDistanceScale);
 	}
 
-	if(PhysMesh.HasBackStops())
+	if(BackstopDistances.Num() && BackstopRadiuses.Num())
 	{
 		nv::cloth::Range<physx::PxVec4> SeparationConstraints = CurrentCloth->getSeparationConstraints();
 		const int32 NumSeparationConstraints = CurrentCloth->getNumSeparationConstraints();
 		check(NumSeparationConstraints <= SkinnedPhysicsMeshNormals.Num());
+		check(BackstopDistances.Num() == BackstopRadiuses.Num());
 
 		for(int32 ConstraintIndex = 0; ConstraintIndex < NumSeparationConstraints; ++ConstraintIndex)
 		{
@@ -1780,7 +1825,13 @@ void FClothingActorNv::UpdateWind(FClothingSimulationContextNv* InContext, const
 			TArray<FVector>& ParticleVelocities = Scratch.ParticleVelocities;
 			CalculateParticleVelocities(ParticleVelocities);
 
-			const TArray<float>& MaxDistances = AssetCreatedFrom->LodData[CurrentLodIndex].PhysicalMeshData.MaxDistances;
+			const UClothPhysicalMeshDataBase* MeshDataBase = AssetCreatedFrom->ClothLodData[CurrentLodIndex]->PhysicalMeshData;
+			const TArray<float>* MaxDistancesPtr = MeshDataBase ? MeshDataBase->GetFloatArray((int32)MaskTarget_PhysMesh::MaxDistance) : nullptr;
+			if (!MaxDistancesPtr || !MaxDistancesPtr->Num())
+			{
+				return;
+			}
+			const TArray<float>& MaxDistances = *MaxDistancesPtr;
 
 			const int32 NumAccelerations = LodData[CurrentLodIndex].Cloth->getNumParticleAccelerations();
 			nv::cloth::Range<physx::PxVec4> ParticleAccelerations = LodData[CurrentLodIndex].Cloth->getParticleAccelerations();
@@ -1830,7 +1881,7 @@ void FClothingActorNv::ConditionalRebuildCollisions()
 	AggregatedCollisions.Reset();
 
 	// Asset-embedded collisions (created during import)
-	AggregatedCollisions.Append(AssetCreatedFrom->LodData[CurrentLodIndex].CollisionData);
+	AggregatedCollisions.Append(AssetCreatedFrom->ClothLodData[CurrentLodIndex]->CollisionData);
 	// Extracted collisions from the physics asset selected by the user
 	AggregatedCollisions.Append(ExtractedCollisions);
 	// External collisions added from the world
@@ -1866,11 +1917,15 @@ void FClothingActorNv::UpdateAnimDrive(FClothingSimulationContextNv* InContext)
 {
 	SCOPE_CYCLE_COUNTER(STAT_NvClothUpdateAnimDrive);
 
-	FClothPhysicalMeshData& PhysMesh = AssetCreatedFrom->LodData[CurrentLodIndex].PhysicalMeshData;
-
-	if(PhysMesh.HasAnimDrive())
+	const UClothPhysicalMeshDataBase* MeshDataBase = AssetCreatedFrom->ClothLodData[CurrentLodIndex]->PhysicalMeshData;
+	if (!MeshDataBase)
+	{
+		return;
+	}
+	const TArray<float>* AnimDriveMultipliers = MeshDataBase->GetFloatArray((int32)MaskTarget_PhysMesh::AnimDriveMultiplier);
+	if(AnimDriveMultipliers && AnimDriveMultipliers->Num())
 	{
-		TArray<float>& AnimDrives = PhysMesh.AnimDriveMultipliers;
+		const TArray<float>& AnimDrives = *AnimDriveMultipliers;
 		const TArray<FVector> CurrentSkinPositions = GetCurrentSkinnedPositions();
 		const TArray<FVector> PreviousSkinPositions = GetPreviousSkinnedPositions();
 
diff --git a/Engine/Source/Runtime/ClothingSystemRuntimeNv/Public/Assets/ClothingAssetNv.h b/Engine/Source/Runtime/ClothingSystemRuntimeNv/Public/Assets/ClothingAssetNv.h
new file mode 100644
--- /dev/null
+++ b/Engine/Source/Runtime/ClothingSystemRuntimeNv/Public/Assets/ClothingAssetNv.h
@@ -0,0 +1,45 @@
+// Copyright 1998-2019 Epic Games, Inc. All Rights Reserved.
+#pragma once
+
+#include "ClothingAsset.h"
+#include "ClothConfigNv.h"
+#include "ClothLODDataNv.h"
+
+#include "ClothingAssetNv.generated.h"
+
+/**
+ * NvCloth implementation of a clothing asset.
+ *
+ * Note: BaseEngine.ini remaps a class named "ClothingAsset" to this class.
+ */
+UCLASS(hidecategories = Object, BlueprintType)
+class CLOTHINGSYSTEMRUNTIMENV_API UClothingAssetNv : public UClothingAssetCommon
+{
+	GENERATED_BODY()
+public:
+	UClothingAssetNv(const FObjectInitializer& ObjectInitializer);
+
+	virtual void PostLoad() override;
+	virtual void Serialize(FArchive& Ar) override;
+
+#if WITH_EDITOR
+	virtual void InvalidateCachedData() override;
+	virtual int32 AddNewLod() override;
+	virtual void PostPropertyChangeCb(const FPropertyChangedEvent& Event) override;
+#endif // WITH_EDITOR
+
+	/** 
+	 * Deprecated property for transitioning the \c FClothConfig struct to the 
+	 * \c UClothConfig class, in a new property called \c ClothSimConfig.
+	 */
+	UPROPERTY()
+	FClothConfig ClothConfig_DEPRECATED;
+
+	/** 
+	 * Deprecasted property for transitioning \c FClothLODData struct to the 
+	 * \c UClothLODDataNv class, in a new property called \c ClothLodData.
+	 */
+	UPROPERTY()
+	TArray<FClothLODData> LodData_DEPRECATED;
+};
+
diff --git a/Engine/Source/Runtime/ClothingSystemRuntimeNv/Public/ClothLODDataNv.h b/Engine/Source/Runtime/ClothingSystemRuntimeNv/Public/ClothLODDataNv.h
new file mode 100644
--- /dev/null
+++ b/Engine/Source/Runtime/ClothingSystemRuntimeNv/Public/ClothLODDataNv.h
@@ -0,0 +1,89 @@
+// Copyright 1998-2019 Epic Games, Inc. All Rights Reserved.
+#pragma once
+
+#include "ClothLODData.h"
+#include "ClothPhysicalMeshDataNv.h"
+#include "ClothParameterMask_PhysMesh.h"
+
+#include "ClothCollisionData.h"
+#include "PointWeightMap.h"
+
+#include "SkeletalMeshTypes.h"
+#include "Containers/Array.h"
+
+#include "ClothLODDataNv.generated.h"
+
+/** */
+UCLASS()
+class CLOTHINGSYSTEMRUNTIMENV_API UClothLODDataNv : public UClothLODDataBase
+{
+	GENERATED_BODY()
+public:
+	UClothLODDataNv(const FObjectInitializer& Init);
+	virtual ~UClothLODDataNv();
+};
+
+
+/* Deprecated.  Use UClothLODDataNv instead. */
+USTRUCT()
+struct CLOTHINGSYSTEMRUNTIMENV_API FClothLODData
+{
+	GENERATED_BODY()
+
+	// Raw phys mesh data
+	UPROPERTY(EditAnywhere, Category = SimMesh)
+	FClothPhysicalMeshData PhysicalMeshData;
+
+	// Collision primitive and covex data for clothing collisions
+	UPROPERTY(EditAnywhere, Category = Collision)
+	FClothCollisionData CollisionData;
+
+#if WITH_EDITORONLY_DATA
+	// Parameter masks defining the physics mesh masked data
+	UPROPERTY(EditAnywhere, Category = Masks)
+	TArray<FClothParameterMask_PhysMesh> ParameterMasks;
+#endif // WITH_EDITORONLY_DATA
+
+	// Skinning data for transitioning from a higher detail LOD to this one
+	TArray<FMeshToMeshVertData> TransitionUpSkinData;
+
+	// Skinning data for transitioning from a lower detail LOD to this one
+	TArray<FMeshToMeshVertData> TransitionDownSkinData;
+
+	bool Serialize(FArchive& Ar)
+	{
+		// Serialize normal tagged data
+		if (!Ar.IsCountingMemory())
+		{
+			UScriptStruct* Struct = FClothLODData::StaticStruct();
+			Struct->SerializeTaggedProperties(Ar, (uint8*)this, Struct, nullptr);
+		}
+		// Serialize the mesh to mesh data (not a USTRUCT)
+		Ar	<< TransitionUpSkinData
+			<< TransitionDownSkinData;
+		return true;
+	}
+
+	void MigrateTo(UClothLODDataNv* LodData) const
+	{
+		PhysicalMeshData.MigrateTo(LodData->PhysicalMeshData);
+		LodData->CollisionData = CollisionData;
+#if WITH_EDITORONLY_DATA
+		LodData->ParameterMasks.SetNum(ParameterMasks.Num());
+		for (int i = 0; i < ParameterMasks.Num(); i++)
+		{
+			ParameterMasks[i].MigrateTo(&LodData->ParameterMasks[i]);
+		}
+#endif // WITH_EDITORONLY_DATA
+		LodData->TransitionUpSkinData = TransitionUpSkinData;
+		LodData->TransitionDownSkinData = TransitionDownSkinData;
+	}
+};
+template<>
+struct TStructOpsTypeTraits<FClothLODData> : public TStructOpsTypeTraitsBase2<FClothLODData>
+{
+	enum
+	{
+		WithSerializer = true,
+	};
+};
diff --git a/Engine/Source/Runtime/ClothingSystemRuntimeNv/Public/ClothParameterMask_PhysMesh.h b/Engine/Source/Runtime/ClothingSystemRuntimeNv/Public/ClothParameterMask_PhysMesh.h
new file mode 100644
--- /dev/null
+++ b/Engine/Source/Runtime/ClothingSystemRuntimeNv/Public/ClothParameterMask_PhysMesh.h
@@ -0,0 +1,53 @@
+// Copyright 1998-2019 Epic Games, Inc. All Rights Reserved.
+#pragma once
+
+#include "PointWeightMap.h"
+#include "ClothPhysicalMeshDataNv.h"
+#include "ClothParameterMask_PhysMesh.generated.h"
+
+/** Deprecated.  Use FPointWeightMap instead. */
+USTRUCT()
+struct CLOTHINGSYSTEMRUNTIMENV_API FClothParameterMask_PhysMesh
+{
+	GENERATED_BODY();
+
+	FClothParameterMask_PhysMesh()
+		: MaskName(NAME_None)
+		, CurrentTarget(MaskTarget_PhysMesh::None)
+		, MaxValue_DEPRECATED(0.0)
+		, MinValue_DEPRECATED(100.0)
+		, bEnabled(false)
+	{}
+
+	void MigrateTo(FPointWeightMap* Weights) const
+	{
+		Weights->Name = MaskName;
+		Weights->CurrentTarget = static_cast<uint8>(CurrentTarget);
+		Weights->Values = Values;
+		Weights->bEnabled = bEnabled;
+	}
+
+	/** Name of the mask, mainly for users to differentiate */
+	UPROPERTY()
+	FName MaskName;
+
+	/** The currently targeted parameter for the mask */
+	UPROPERTY()
+	MaskTarget_PhysMesh CurrentTarget;
+
+	/** The maximum value currently in the mask value array */
+	UPROPERTY()
+	float MaxValue_DEPRECATED;
+
+	/** The maximum value currently in the mask value array */
+	UPROPERTY()
+	float MinValue_DEPRECATED;
+
+	/** The actual values stored in the mask */
+	UPROPERTY()
+	TArray<float> Values;
+
+	/** Whether this mask is enabled and able to effect final mesh values */
+	UPROPERTY()
+	bool bEnabled;
+};
diff --git a/Engine/Source/Runtime/ClothingSystemRuntimeNv/Public/ClothPhysicalMeshDataNv.h b/Engine/Source/Runtime/ClothingSystemRuntimeNv/Public/ClothPhysicalMeshDataNv.h
new file mode 100644
--- /dev/null
+++ b/Engine/Source/Runtime/ClothingSystemRuntimeNv/Public/ClothPhysicalMeshDataNv.h
@@ -0,0 +1,162 @@
+// Copyright 1998-2019 Epic Games, Inc. All Rights Reserved.
+#pragma once
+
+#include "ClothPhysicalMeshData.h"
+#include "ClothPhysicalMeshDataNv.generated.h"
+
+/** The possible targets for a physical mesh point weight map. */
+UENUM()
+enum class MaskTarget_PhysMesh : uint8
+{
+	None = 0,			// None should always be zero
+	MaxDistance,
+	BackstopDistance,
+	BackstopRadius,
+	AnimDriveMultiplier
+};
+
+/** NV specific spatial simulation data for a mesh. */
+UCLASS()
+class CLOTHINGSYSTEMRUNTIMENV_API UClothPhysicalMeshDataNv : public UClothPhysicalMeshDataBase
+{
+	GENERATED_BODY()
+public:
+	UClothPhysicalMeshDataNv();
+	virtual ~UClothPhysicalMeshDataNv();
+
+	void
+	Reset(const int32 NumPoints) override;
+
+	void
+	ClearParticleParameters() override;
+
+	virtual void
+	BuildSelfCollisionData(const UClothConfigBase* ClothConfig);
+
+	virtual UEnum* GetFloatArrayTargets() const
+	{ return StaticEnum<MaskTarget_PhysMesh>(); }
+
+	virtual bool IsFullyKinematic() const override
+	{ return MaxDistances.Num() == 0; }
+	
+	virtual bool IsFixed(const uint16 X, const float Threshold=0.1f) const override
+	{ return MaxDistances.IsValidIndex(X) && MaxDistances[X] <= Threshold; }
+
+	virtual bool IsFixed(const uint16 X, const uint16 Y, const uint16 Z, const float Threshold=0.1f) const override
+	{ return IsFixed(X, Threshold) || IsFixed(Y, Threshold) || IsFixed(Z, Threshold); }
+
+	// The distance that each vertex can move away from its reference (skinned) position
+	UPROPERTY(EditAnywhere, Category = SimMesh)
+	TArray<float> MaxDistances;
+
+	// Distance along the plane of the surface that the particles can travel (separation constraint)
+	UPROPERTY(EditAnywhere, Category = SimMesh)
+	TArray<float> BackstopDistances;
+
+	// Radius of movement to allow for backstop movement
+	UPROPERTY(EditAnywhere, Category = SimMesh)
+	TArray<float> BackstopRadiuses;
+
+	// Strength of anim drive per-particle (spring driving particle back to skinned location
+	UPROPERTY(EditAnywhere, Category = SimMesh)
+	TArray<float> AnimDriveMultipliers;
+};
+
+
+/** Deprecated.  Use UClothPhysicalMeshDataNv instead. */
+USTRUCT()
+struct CLOTHINGSYSTEMRUNTIMENV_API FClothPhysicalMeshData
+{
+	GENERATED_BODY()
+
+	void Reset(const int32 InNumVerts);
+
+	// Clear out any target properties in this physical mesh
+	void ClearParticleParameters();
+
+	// Whether the mesh uses backstops
+	bool HasBackStops() const;
+
+	// Whether the mesh uses anim drives
+	bool HasAnimDrive() const;
+
+	// Positions of each simulation vertex
+	UPROPERTY(EditAnywhere, Category = SimMesh)
+	TArray<FVector> Vertices;
+
+	// Normal at each vertex
+	UPROPERTY(EditAnywhere, Category = SimMesh)
+	TArray<FVector> Normals;
+
+#if WITH_EDITORONLY_DATA
+	// Color at each vertex
+	UPROPERTY(EditAnywhere, Category = SimMesh)
+	TArray<FColor> VertexColors;
+#endif // WITH_EDITORONLY_DATA
+
+	// Indices of the simulation mesh triangles
+	UPROPERTY(EditAnywhere, Category = SimMesh)
+	TArray<uint32> Indices;
+
+	// The distance that each vertex can move away from its reference (skinned) position
+	UPROPERTY(EditAnywhere, Category = SimMesh)
+	TArray<float> MaxDistances;
+
+	// Distance along the plane of the surface that the particles can travel (separation constraint)
+	UPROPERTY(EditAnywhere, Category = SimMesh)
+	TArray<float> BackstopDistances;
+
+	// Radius of movement to allow for backstop movement
+	UPROPERTY(EditAnywhere, Category = SimMesh)
+	TArray<float> BackstopRadiuses;
+
+	// Strength of anim drive per-particle (spring driving particle back to skinned location
+	UPROPERTY(EditAnywhere, Category = SimMesh)
+	TArray<float> AnimDriveMultipliers;
+
+	// Inverse mass for each vertex in the physical mesh
+	UPROPERTY(EditAnywhere, Category = SimMesh)
+	TArray<float> InverseMasses;
+
+	// Indices and weights for each vertex, used to skin the mesh to create the reference pose
+	UPROPERTY(EditAnywhere, Category = SimMesh)
+	TArray<FClothVertBoneData> BoneData;
+
+	// Maximum number of bone weights of any vetex
+	UPROPERTY(EditAnywhere, Category = SimMesh)
+	int32 MaxBoneWeights;
+
+	// Number of fixed verts in the simulation mesh (fixed verts are just skinned and do not simulate)
+	UPROPERTY(EditAnywhere, Category = SimMesh)
+	int32 NumFixedVerts;
+
+	// Valid indices to use for self collisions (reduced set of Indices)
+	UPROPERTY(EditAnywhere, Category = SimMesh)
+	TArray<uint32> SelfCollisionIndices;
+
+	void MigrateTo(UClothPhysicalMeshDataBase* MeshData) const
+	{
+		MeshData->Vertices = Vertices;
+		MeshData->Normals = Normals;
+#if WITH_EDITORONLY_DATA
+		MeshData->VertexColors = VertexColors;
+#endif
+		MeshData->Indices = Indices;
+		MeshData->InverseMasses = InverseMasses;
+		MeshData->BoneData = BoneData;
+		MeshData->NumFixedVerts = NumFixedVerts;
+		MeshData->MaxBoneWeights = MaxBoneWeights;
+		MeshData->SelfCollisionIndices = SelfCollisionIndices;
+		if (UClothPhysicalMeshDataNv* NvMeshData = Cast<UClothPhysicalMeshDataNv>(MeshData))
+		{
+			NvMeshData->MaxDistances = MaxDistances;
+			NvMeshData->BackstopDistances = BackstopDistances;
+			NvMeshData->BackstopRadiuses = BackstopRadiuses;
+			NvMeshData->AnimDriveMultipliers = AnimDriveMultipliers;
+		}
+		else
+		{
+			unimplemented();
+		}
+	}
+};
diff --git a/Engine/Source/Runtime/ClothingSystemRuntimeNv/Public/ClothingActor.h b/Engine/Source/Runtime/ClothingSystemRuntimeNv/Public/ClothingActor.h
--- a/Engine/Source/Runtime/ClothingSystemRuntimeNv/Public/ClothingActor.h
+++ b/Engine/Source/Runtime/ClothingSystemRuntimeNv/Public/ClothingActor.h
@@ -2,7 +2,7 @@
 
 #pragma once
 
-class UClothingAsset;
+class UClothingAssetCommon;
 
 // Base clothing actor just needs a link back to its parent asset, everything else defined by derived simulation
 class FClothingActorBase
@@ -10,5 +10,5 @@ class FClothingActorBase
 public:
 
 	// Link back to parent asset
-	UClothingAsset* AssetCreatedFrom;
+	UClothingAssetCommon* AssetCreatedFrom;
 };
diff --git a/Engine/Source/Runtime/ClothingSystemRuntimeNv/Public/ClothingSimulationNv.h b/Engine/Source/Runtime/ClothingSystemRuntimeNv/Public/ClothingSimulationNv.h
--- a/Engine/Source/Runtime/ClothingSystemRuntimeNv/Public/ClothingSimulationNv.h
+++ b/Engine/Source/Runtime/ClothingSystemRuntimeNv/Public/ClothingSimulationNv.h
@@ -5,10 +5,11 @@
 #if WITH_NVCLOTH
 
 #include "ClothingSimulation.h"
-#include "ClothingActor.h"
 #include "ClothingSimulationContext.h"
+#include "ClothingActor.h"
+
+#include "Assets/ClothingAssetNv.h"
 
-#include "Assets/ClothingAsset.h"
 
 namespace nv
 {
@@ -51,6 +52,7 @@ struct FClothingActorScratchData
 class FClothingActorNv final : public FClothingActorBase
 {
 public:
+	using FClothingActorBase::AssetCreatedFrom;
 
 	FClothingActorNv();
 
@@ -251,11 +253,11 @@ private:
 
 	// Given a clothing config from an asset, apply it to the provided actor. Currently
 	// this is only used from CreateActor, but could be exposed for runtime changes
-	void ApplyClothConfig(FClothConfig &Config, FClothingActorNv &InActor);
+	void ApplyClothConfig(UClothConfigBase *BaseConfig, FClothingActorNv &InActor);
 
 	// Extract collisions from the physics asset inside Asset and apply them to InActor
 	// Not safe to call from workers (i.e. inside the simulation).
-	void ExtractActorCollisions(UClothingAsset* Asset, FClothingActorNv &InActor);
+	void ExtractActorCollisions(UClothingAssetCommon* Asset, FClothingActorNv &InActor);
 
 	// The current LOD index for the owning skeletal mesh component
 	int32 CurrentMeshLodIndex;
diff --git a/Engine/Source/Runtime/Core/Private/HAL/LowLevelMemTracker.cpp b/Engine/Source/Runtime/Core/Private/HAL/LowLevelMemTracker.cpp
--- a/Engine/Source/Runtime/Core/Private/HAL/LowLevelMemTracker.cpp
+++ b/Engine/Source/Runtime/Core/Private/HAL/LowLevelMemTracker.cpp
@@ -84,7 +84,9 @@ DECLARE_LLM_MEMORY_STAT(TEXT("Particles"), STAT_ParticlesLLM, STATGROUP_LLMFULL)
 DECLARE_LLM_MEMORY_STAT(TEXT("Niagara"), STAT_NiagaraLLM, STATGROUP_LLMFULL);
 DECLARE_LLM_MEMORY_STAT(TEXT("GC"), STAT_GCLLM, STATGROUP_LLMFULL);
 DECLARE_LLM_MEMORY_STAT(TEXT("UI"), STAT_UILLM, STATGROUP_LLMFULL);
+DECLARE_LLM_MEMORY_STAT(TEXT("Physics"), STAT_PhysicsLLM, STATGROUP_LLMFULL);
 DECLARE_LLM_MEMORY_STAT(TEXT("PhysX"), STAT_PhysXLLM, STATGROUP_LLMFULL);
+DECLARE_LLM_MEMORY_STAT(TEXT("Chaos"), STAT_ChaosLLM, STATGROUP_LLMFULL);
 DECLARE_LLM_MEMORY_STAT(TEXT("EnginePreInit"), STAT_EnginePreInitLLM, STATGROUP_LLMFULL);
 DECLARE_LLM_MEMORY_STAT(TEXT("EngineInit"), STAT_EngineInitLLM, STATGROUP_LLMFULL);
 DECLARE_LLM_MEMORY_STAT(TEXT("Rendering Thread"), STAT_RenderingThreadLLM, STATGROUP_LLMFULL);
@@ -113,7 +115,9 @@ DECLARE_LLM_MEMORY_STAT(TEXT("VirtualMemory"), STAT_PlatformVMLLM, STATGROUP_LLM
 DECLARE_LLM_MEMORY_STAT(TEXT("Total"), STAT_TrackedTotalSummaryLLM, STATGROUP_LLM);
 DECLARE_LLM_MEMORY_STAT(TEXT("Audio"), STAT_AudioSummaryLLM, STATGROUP_LLM);
 DECLARE_LLM_MEMORY_STAT(TEXT("Meshes"), STAT_MeshesSummaryLLM, STATGROUP_LLM);
+DECLARE_LLM_MEMORY_STAT(TEXT("Physics"), STAT_PhysicsSummaryLLM, STATGROUP_LLM);
 DECLARE_LLM_MEMORY_STAT(TEXT("PhysX"), STAT_PhysXSummaryLLM, STATGROUP_LLM);
+DECLARE_LLM_MEMORY_STAT(TEXT("Chaos"), STAT_ChaosSummaryLLM, STATGROUP_LLM);
 DECLARE_LLM_MEMORY_STAT(TEXT("UObject"), STAT_UObjectSummaryLLM, STATGROUP_LLM);
 DECLARE_LLM_MEMORY_STAT(TEXT("Animation"), STAT_AnimationSummaryLLM, STATGROUP_LLM);
 DECLARE_LLM_MEMORY_STAT(TEXT("StaticMesh"), STAT_StaticMeshSummaryLLM, STATGROUP_LLM);
diff --git a/Engine/Source/Runtime/Core/Private/HAL/LowLevelMemoryUtils.h b/Engine/Source/Runtime/Core/Private/HAL/LowLevelMemoryUtils.h
--- a/Engine/Source/Runtime/Core/Private/HAL/LowLevelMemoryUtils.h
+++ b/Engine/Source/Runtime/Core/Private/HAL/LowLevelMemoryUtils.h
@@ -53,7 +53,15 @@ public:
 	void Add(const T& Item)
 	{
 		if (Count == Capacity)
-			Reserve(Capacity ? (150U * Capacity) / 100U : DefaultCapacity);
+		{
+			uint32 NewCapacity = DefaultCapacity;
+			if (Capacity)
+			{
+				NewCapacity = (15U * Capacity) / 10U;
+				ensureMsgf(NewCapacity > Capacity, TEXT("Unsigned integer overflow."));
+			}
+			Reserve(NewCapacity);
+		}
 
 		Array[Count] = Item;
 		++Count;
diff --git a/Engine/Source/Runtime/Core/Private/UObject/DevObjectVersion.cpp b/Engine/Source/Runtime/Core/Private/UObject/DevObjectVersion.cpp
--- a/Engine/Source/Runtime/Core/Private/UObject/DevObjectVersion.cpp
+++ b/Engine/Source/Runtime/Core/Private/UObject/DevObjectVersion.cpp
@@ -24,6 +24,7 @@
 #include "UObject/AutomationObjectVersion.h"
 #include "UObject/NiagaraObjectVersion.h"
 #include "UObject/DestructionObjectVersion.h"
+#include "UObject/ExternalPhysicsCustomObjectVersion.h"
 #include "UObject/CineCameraObjectVersion.h"
 
 DEFINE_LOG_CATEGORY_STATIC(LogDevObjectVersion, Log, All);
@@ -165,6 +166,11 @@ const FGuid FDestructionObjectVersion::GUID(0x174F1F0B, 0xB4C645A5, 0xB13F2EE8,
 // Register Destruction custom version with Core
 FDevVersionRegistration GRegisterDestructionObjectVersion(FDestructionObjectVersion::GUID, FDestructionObjectVersion::LatestVersion, TEXT("Dev-Destruction"));
 
+// Unique Physics Object version id
+const FGuid FExternalPhysicsCustomObjectVersion::GUID(0x35F94A83, 0xE258406C, 0xA31809F5, 0x9610247C);
+// Register Physics custom version with Core
+FDevVersionRegistration GRegisterExternalPhysicsCustomVersion(FExternalPhysicsCustomObjectVersion::GUID, FExternalPhysicsCustomObjectVersion::LatestVersion, TEXT("Dev-Physics-Ext"));
+
 // Unique CineCamera Object version id
 const FGuid FCineCameraObjectVersion::GUID(0xB2E18506, 0x4273CFC2, 0xA54EF4BB, 0x758BBA07);
 // Register CineCamera custom version with Core
diff --git a/Engine/Source/Runtime/Core/Public/UObject/PhysicsObjectVersion.h b/Engine/Source/Runtime/Core/Public/UObject/PhysicsObjectVersion.h
--- a/Engine/Source/Runtime/Core/Public/UObject/PhysicsObjectVersion.h
+++ b/Engine/Source/Runtime/Core/Public/UObject/PhysicsObjectVersion.h
@@ -11,7 +11,10 @@ struct CORE_API FPhysicsObjectVersion
 	{
 		// Before any version changes were made
 		BeforeCustomVersionWasAdded = 0,
-
+		// Adding PerShapeData to serialization
+		PerShapeData,
+		// Add serialization from handle back to particle
+		SerializeGTGeometryParticles,
 		// Temporary versioning to serialize hair guide data in addition to hair render data
 		HairAssetSerialization,
 		HairAssetSerialization_V2,
diff --git a/Engine/Source/Runtime/Engine/Classes/Animation/Skeleton.h b/Engine/Source/Runtime/Engine/Classes/Animation/Skeleton.h
--- a/Engine/Source/Runtime/Engine/Classes/Animation/Skeleton.h
+++ b/Engine/Source/Runtime/Engine/Classes/Animation/Skeleton.h
@@ -837,8 +837,8 @@ public:
 	ENGINE_API void RemoveBonesFromSkeleton(const TArray<FName>& BonesToRemove, bool bRemoveChildBones);
 
 	// Asset registry information for animation notifies
-	static const FName AnimNotifyTag;
-	static const FString AnimNotifyTagDelimiter;
+	ENGINE_API static const FName AnimNotifyTag;
+	ENGINE_API static const FString AnimNotifyTagDelimiter;
 
 	// Asset registry information for animation curves
 	ENGINE_API static const FName CurveNameTag;
diff --git a/Engine/Source/Runtime/Engine/Private/Animation/AnimNode_CustomProperty.cpp b/Engine/Source/Runtime/Engine/Private/Animation/AnimNode_CustomProperty.cpp
--- a/Engine/Source/Runtime/Engine/Private/Animation/AnimNode_CustomProperty.cpp
+++ b/Engine/Source/Runtime/Engine/Private/Animation/AnimNode_CustomProperty.cpp
@@ -6,7 +6,8 @@
 #include "Animation/AnimNode_LinkedInputPose.h"
 
 FAnimNode_CustomProperty::FAnimNode_CustomProperty()
-	: TargetInstance(nullptr)
+	: FAnimNode_Base()
+	, TargetInstance(nullptr)
 #if WITH_EDITOR
 	, bReinitializeProperties(false)
 #endif // WITH_EDITOR
diff --git a/Engine/Source/Runtime/Engine/Private/Animation/AnimNode_StateMachine.cpp b/Engine/Source/Runtime/Engine/Private/Animation/AnimNode_StateMachine.cpp
--- a/Engine/Source/Runtime/Engine/Private/Animation/AnimNode_StateMachine.cpp
+++ b/Engine/Source/Runtime/Engine/Private/Animation/AnimNode_StateMachine.cpp
@@ -451,10 +451,6 @@ void FAnimNode_StateMachine::Update_AnyThread(const FAnimationUpdateContext& Con
 					{
 						InertializationNode->Request(ReferenceTransition.CrossfadeDuration);
 					}
-					else
-					{
-						UE_LOG(LogAnimation, Error, TEXT("FAnimNode_StateMachine: No inertialization context found for transition"));
-					}
 				}
 
 #if WITH_EDITORONLY_DATA
diff --git a/Engine/Source/Runtime/Engine/Private/Animation/AnimSequence.cpp b/Engine/Source/Runtime/Engine/Private/Animation/AnimSequence.cpp
--- a/Engine/Source/Runtime/Engine/Private/Animation/AnimSequence.cpp
+++ b/Engine/Source/Runtime/Engine/Private/Animation/AnimSequence.cpp
@@ -1672,25 +1672,13 @@ void UAnimSequence::GetBonePose_Additive(FCompactPose& OutPose, FBlendedCurve& O
 
 void UAnimSequence::GetAdditiveBasePose(FCompactPose& OutPose, FBlendedCurve& OutCurve, const FAnimExtractContext& ExtractionContext) const
 {
-	const bool bGetAdditiveBasePoseValid = (RefPoseType == ABPT_RefPose || !RefPoseSeq->IsValidAdditive() || RefPoseSeq->RawAnimationData.Num() > 0);
-	if (!bGetAdditiveBasePoseValid) //If this fails there is not enough information to get the base pose
-	{
-		FString Name = GetName();
-		FString SkelName = GetSkeleton() ? GetSkeleton()->GetName() : TEXT("NoSkeleton");
-		FString RefName = RefPoseSeq->GetName();
-		FString RefSkelName = RefPoseSeq->GetSkeleton() ? RefPoseSeq->GetSkeleton()->GetName() : TEXT("NoRefSeqSkeleton");
-		const TCHAR* NeedsLoad = RefPoseSeq->HasAnyFlags(RF_NeedLoad) ? TEXT("Yes") : TEXT("No");
-
-		checkf(false, TEXT("Cannot get valid base pose for Anim: ['%s' (Skel:%s)] RefSeq: ['%s' (Skel:%s)] RawAnimDataNum: %i NeedsLoad: %s"), *Name, *SkelName, *RefName, *RefSkelName, RefPoseSeq->RawAnimationData.Num(), NeedsLoad);
-	}
-
 	switch (RefPoseType)
 	{
 		// use whole animation as a base pose. Need BasePoseSeq.
 		case ABPT_AnimScaled:
 		{
 			// normalize time to fit base seq
-			const float Fraction = FMath::Clamp<float>(ExtractionContext.CurrentTime / SequenceLength, 0.f, 1.f);
+			const float Fraction = (SequenceLength > 0.f)? FMath::Clamp<float>(ExtractionContext.CurrentTime / SequenceLength, 0.f, 1.f) : 0.f;
 			const float BasePoseTime = RefPoseSeq->SequenceLength * Fraction;
 
 			FAnimExtractContext BasePoseExtractionContext(ExtractionContext);
diff --git a/Engine/Source/Runtime/Engine/Private/Collision/Collision.cpp b/Engine/Source/Runtime/Engine/Private/Collision/Collision.cpp
--- a/Engine/Source/Runtime/Engine/Private/Collision/Collision.cpp
+++ b/Engine/Source/Runtime/Engine/Private/Collision/Collision.cpp
@@ -178,6 +178,10 @@ FCollisionQueryParams::FCollisionQueryParams(FName InTraceTag, const TStatId& In
 	{
 		OwnerTag = InIgnoreActor->GetFName();
 	}
+
+#if !(UE_BUILD_TEST || UE_BUILD_SHIPPING)
+		bDebugQuery = false;
+#endif
 }
 
 
diff --git a/Engine/Source/Runtime/Engine/Private/Collision/CollisionConversions.cpp b/Engine/Source/Runtime/Engine/Private/Collision/CollisionConversions.cpp
--- a/Engine/Source/Runtime/Engine/Private/Collision/CollisionConversions.cpp
+++ b/Engine/Source/Runtime/Engine/Private/Collision/CollisionConversions.cpp
@@ -17,6 +17,7 @@
 #elif WITH_CHAOS
 #include "Physics/Experimental/ChaosInterfaceWrapper.h"
 #include "Physics/Experimental/PhysInterface_Chaos.h"
+#include "Chaos/ParticleHandle.h"
 #endif
 #include "PhysicsEngine/CollisionQueryFilterCallback.h"
 
@@ -381,7 +382,10 @@ EConvertQueryResult ConvertTraceResults(bool& OutHasValidBlockingHit, const UWor
 	const FVector Dir = (EndLoc - StartLoc).GetSafeNormal();
 	if (TIsSame<Hit, FHitSweep>::Value)
 	{
-		SetInternalFaceIndex(Hits[0], FindFaceIndex(Hits[0], Dir));
+		if (!HadInitialOverlap(Hits[0]))
+		{
+			SetInternalFaceIndex(Hits[0], FindFaceIndex(Hits[0], Dir));
+		}
 	}
 	EConvertQueryResult Result = ConvertQueryImpactHit(World, Hits[0], OutHit, CheckLength, QueryFilter, StartLoc, EndLoc, &Geom, QueryTM, bReturnFaceIndex, bReturnPhysMat);
 	OutHasValidBlockingHit = Result == EConvertQueryResult::Valid;
@@ -509,17 +513,13 @@ void ConvertQueryOverlap(const FPhysicsShape& Shape, const FPhysicsActor& Actor,
 	// Try body instance
 	if (const FBodyInstance* BodyInst = GetUserData(Actor))
 	{
-#if WITH_CHAOS || WITH_IMMEDIATE_PHYSX
-        ensure(false);
-#else
-        BodyInst = FPhysicsInterface_PhysX::ShapeToOriginalBodyInstance(BodyInst, &Shape);
+        BodyInst = FPhysicsInterface::ShapeToOriginalBodyInstance(BodyInst, &Shape);
 		if (const UPrimitiveComponent* OwnerComponent = BodyInst->OwnerComponent.Get())
 		{
 			OutOverlap.Actor = OwnerComponent->GetOwner();
 			OutOverlap.Component = BodyInst->OwnerComponent; // Copying weak pointer is faster than assigning raw pointer.
 			OutOverlap.ItemIndex = OwnerComponent->bMultiBodyOverlap ? BodyInst->InstanceBodyIndex : INDEX_NONE;
 		}
-#endif
 	}
 #if PHYSICS_INTERFACE_PHYSX
 	else if(const FCustomPhysXPayload* CustomPayload = GetUserData<FCustomPhysXPayload>(Shape))
diff --git a/Engine/Source/Runtime/Engine/Private/Collision/CollisionConversions.h b/Engine/Source/Runtime/Engine/Private/Collision/CollisionConversions.h
--- a/Engine/Source/Runtime/Engine/Private/Collision/CollisionConversions.h
+++ b/Engine/Source/Runtime/Engine/Private/Collision/CollisionConversions.h
@@ -54,7 +54,7 @@ template <typename Hit>
 EConvertQueryResult ConvertTraceResults(bool& OutHasValidBlockingHit, const UWorld* World, int32 NumHits, Hit* Hits, float CheckLength, const FCollisionFilterData& QueryFilter, FHitResult& OutHits, const FVector& StartLoc, const FVector& EndLoc, const FPhysicsGeometry& Geom, const FTransform& QueryTM, float MaxDistance, bool bReturnFaceIndex, bool bReturnPhysMat);
 
 /** 
- * Util to convert physX overlap query to our overlap result
+ * Util to convert physX or Chaos overlap query to our overlap result
  *
  * @param	PShape		Shape that overlaps
  * @param	PActor		Specific actor as PShape might be shared among many actors
diff --git a/Engine/Source/Runtime/Engine/Private/Collision/CollisionDebugDrawing.h b/Engine/Source/Runtime/Engine/Private/Collision/CollisionDebugDrawing.h
--- a/Engine/Source/Runtime/Engine/Private/Collision/CollisionDebugDrawing.h
+++ b/Engine/Source/Runtime/Engine/Private/Collision/CollisionDebugDrawing.h
@@ -9,6 +9,7 @@
 #include "EngineDefines.h"
 #include "PhysicsPublic.h"
 #include "PhysXIncludes.h"
+#include "Chaos/ImplicitObject.h"
 
 #if WITH_PHYSX
 /** Draw PhysX geom with overlaps */
@@ -17,7 +18,5 @@ void DrawGeomOverlaps(const UWorld* InWorld, const PxGeometry& PGeom, const PxTr
 void DrawGeomSweeps(const UWorld* InWorld, const FVector& Start, const FVector& End, const PxGeometry& PGeom, const PxQuat& Rotation, const TArray<FHitResult>& Hits, float Lifetime);
 #endif
 
-#if INCLUDE_CHAOS
 void DrawGeomOverlaps(const UWorld* InWorld, const Chaos::TImplicitObject<float, 3>& Geom, const FTransform& GeomPose, TArray<struct FOverlapResult>& Overlaps, float Lifetime);
 void DrawGeomSweeps(const UWorld* InWorld, const FVector& Start, const FVector& End, const Chaos::TImplicitObject<float,3>& Geom, const FQuat& Rotation, const TArray<FHitResult>& Hits, float Lifetime);
-#endif
