Cloth Android/ios platform compile error fixed when chaos is enabled

diff --git a/Engine/Source/Runtime/ClothingSystemRuntimeNv/Private/Assets/ClothingAssetNv.cpp b/Engine/Source/Runtime/ClothingSystemRuntimeNv/Private/Assets/ClothingAssetNv.cpp
--- a/Engine/Source/Runtime/ClothingSystemRuntimeNv/Private/Assets/ClothingAssetNv.cpp
+++ b/Engine/Source/Runtime/ClothingSystemRuntimeNv/Private/Assets/ClothingAssetNv.cpp
@@ -14,6 +14,7 @@
 #include "PhysicsPublic.h"
 #include "Modules/ModuleManager.h"
 #include "Widgets/Notifications/SNotificationList.h"
+#include "Features/IModularFeatures.h"
 #include "Framework/Notifications/NotificationManager.h"
 #include "UObject/AnimPhysObjectVersion.h"
 #include "Rendering/SkeletalMeshModel.h"
