Move common clothing asset functionalities to ClothingSystemRuntimeCommon to enable Chaos Cloth authoring.
Create a new ClothingSimulationFactory function to get the definitive default factory, for whatever cloth factories are actually enabled.
Rework nv cloth/apex dependencies in the ClothingSystemRuntimeNv module to allow the move of the legacy components and common config elements to the common module.
Add legacy cloth asset and runtime components to the common module to allow legacy assets to be loaded in a Chaos simulation, even with nv cloth disabled.
Allow storage of several configs to coexist in the ClothingAsset object config map, with their names as the key.
Resurrect the cloth physical mesh data structure to avoid the UClass overhead.
Add migration functions and BaseEngine redirectors to keep backward compatibility with legacy cloth assets, and assets saved after Ryan's refactor.

diff --git a/Engine/Config/BaseEngine.ini b/Engine/Config/BaseEngine.ini
--- a/Engine/Config/BaseEngine.ini
+++ b/Engine/Config/BaseEngine.ini
@@ -700,15 +700,30 @@ bUseStreamingPause=false
 
 ; Rename ClothingSystemRuntime to ClothingSystemRuntimeNv, base functionality moved to ClothingSystemRuntimeInterface.
 +PackageRedirects=(OldName="/Script/ClothingSystemRuntime",NewName="/Script/ClothingSystemRuntimeNv")
-+StructRedirects=(OldName="/Script/ClothingSystemRuntimeNv.ClothConfig",NewName="/Script/ClothingSystemRuntimeNv.ClothConfig_Legacy")
 +ClassRedirects=(OldName="/Script/ClothingSystemRuntime.ClothingSimulationFactoryNv",NewName="/Script/ClothingSystemRuntimeNv.ClothingSimulationFactoryNv")
 +ClassRedirects=(OldName="/Script/ClothingSystemRuntime.ClothingSimulationInteractorNv",NewName="/Script/ClothingSystemRuntimeNv.ClothingSimulationInteractorNv")
-+ClassRedirects=(OldName="/Script/ClothingSystemRuntime.ClothingAsset",NewName="/Script/ClothingSystemRuntimeNv.ClothingAssetNv")
 +FunctionRedirects=(OldName="ClothingSystemRuntime.ClothingSimulationInteractorNv.SetAnimDriveSpringStiffness",NewName="ClothingSystemRuntimeNv.ClothingSimulationInteractorNv.SetAnimDriveSpringStiffness")
 +FunctionRedirects=(OldName="ClothingSystemRuntime.ClothingSimulationInteractorNv.SetAnimDriveDamperStiffness",NewName="ClothingSystemRuntimeNv.ClothingSimulationInteractorNv.SetAnimDriveDamperStiffness")
 +FunctionRedirects=(OldName="ClothingSystemRuntime.ClothingSimulationInteractorNv.EnableGravityOverride",NewName="ClothingSystemRuntimeNv.ClothingSimulationInteractorNv.EnableGravityOverride")
 +FunctionRedirects=(OldName="ClothingSystemRuntime.ClothingSimulationInteractorNv.DisableGravityOverride",NewName="ClothingSystemRuntimeNv.ClothingSimulationInteractorNv.DisableGravityOverride")
 
+; Redirect legacy ClothingSystemRuntime components to ClothingSystemRuntimeCommon
++ClassRedirects=(OldName="ClothingAsset",NewName="/Script/ClothingSystemRuntimeCommon.ClothingAssetCommon")
++StructRedirects=(OldName="ClothLODData",NewName="/Script/ClothingSystemRuntimeCommon.ClothLODData_Legacy")
++StructRedirects=(OldName="ClothConfig",NewName="/Script/ClothingSystemRuntimeCommon.ClothConfig_Legacy")
++StructRedirects=(OldName="ClothParameterMask_PhysMesh",NewName="/Script/ClothingSystemRuntimeCommon.ClothParameterMask_Legacy")
++StructRedirects=(OldName="ClothConstraintSetup",NewName="/Script/ClothingSystemRuntimeCommon.ClothConstraintSetup_Legacy")
++EnumRedirects=(OldName="EClothingWindMethod",NewName="/Script/ClothingSystemRuntimeCommon.EClothingWindMethod_Legacy")
+
+; Redirect common ClothingSystemRuntimeNv components to ClothingSystemRuntimeCommon
++ClassRedirects=(OldName="/Script/ClothingSystemRuntimeNv.ClothingAssetNv",NewName="/Script/ClothingSystemRuntimeCommon.ClothingAssetCommon")
++ClassRedirects=(OldName="/Script/ClothingSystemRuntimeNv.ClothLODDataNv",NewName="/Script/ClothingSystemRuntimeCommon.ClothLODDataCommon")
++EnumRedirects=(OldName="MaskTarget_PhysMesh",NewName="/Script/ClothingSystemRuntimeCommon.EWeightMapTargetCommon")
+
+; Deprecate the ClothPhysicalMeshDataNv class, but keep it around as a legacy class for clothing asset backward compatibility
++ClassRedirects=(OldName="/Script/ClothingSystemRuntimeInterface.ClothPhysicalMeshDataBase",NewName="/Script/ClothingSystemRuntimeInterface.ClothPhysicalMeshDataBase_Legacy")
++ClassRedirects=(OldName="/Script/ClothingSystemRuntimeNv.ClothPhysicalMeshDataNv",NewName="/Script/ClothingSystemRuntimeNv.ClothPhysicalMeshDataNv_Legacy")
+
 +FunctionRedirects=(OldName="SkeletalMeshComponent.GetSubInstanceByTag",NewName="SkeletalMeshComponent.GetLinkedAnimGraphInstanceByTag")
 +FunctionRedirects=(OldName="SkeletalMeshComponent.GetSubInstancesByTag",NewName="SkeletalMeshComponent.GetLinkedAnimGraphInstancesByTag")
 +PropertyRedirects=(OldName="SkeletalMeshComponent.GetLinkedAnimGraphInstancesByTag.OutSubInstances",NewName="OutLinkedInstances")
diff --git a/Engine/Plugins/Experimental/ChaosCloth/Source/Chaos/Private/ChaosCloth/ChaosClothingSimulation.cpp b/Engine/Plugins/Experimental/ChaosCloth/Source/Chaos/Private/ChaosCloth/ChaosClothingSimulation.cpp
--- a/Engine/Plugins/Experimental/ChaosCloth/Source/Chaos/Private/ChaosCloth/ChaosClothingSimulation.cpp
+++ b/Engine/Plugins/Experimental/ChaosCloth/Source/Chaos/Private/ChaosCloth/ChaosClothingSimulation.cpp
@@ -34,6 +34,7 @@
 #include "Chaos/Utilities.h"
 #include "Chaos/Vector.h"
 #include "ChaosCloth/ChaosClothConfig.h"
+#include "ChaosCloth/ChaosWeightMapTarget.h"
 #include "ChaosCloth/ChaosClothPrivate.h"
 
 #if WITH_PHYSX && !PLATFORM_LUMIN && !PLATFORM_ANDROID
@@ -160,8 +161,12 @@ void ClothingSimulation::CreateActor(USkeletalMeshComponent* InOwnerComponent, U
 	//Evolution->SetCCD(true); // ryan!!!
 
 	UClothingAssetCommon* Asset = Cast<UClothingAssetCommon>(InAsset);
-	const UChaosClothConfig* const ChaosClothSimConfig = Cast<UChaosClothConfig>(Asset->ChaosClothSimConfig);
-	check(ChaosClothSimConfig);
+	const UChaosClothConfig* const ChaosClothSimConfig = Asset->GetClothConfig<UChaosClothConfig>();
+	if (!ChaosClothSimConfig)
+	{
+		UE_LOG(LogChaosCloth, Warning, TEXT("Missing Chaos config Cloth LOD asset to %s in sim slot %d"), InOwnerComponent->GetOwner() ? *InOwnerComponent->GetOwner()->GetName() : TEXT("None"), InSimDataIndex);
+		return;
+	}
 
 	ClothingSimulationContext Context;
 	FillContext(InOwnerComponent, 0, &Context);
@@ -175,9 +180,8 @@ void ClothingSimulation::CreateActor(USkeletalMeshComponent* InOwnerComponent, U
 	AnimDriveSpringStiffness[InSimDataIndex] = ChaosClothSimConfig->AnimDriveSpringStiffness;
 
 	check(Asset->GetNumLods() == 1);
-	UClothLODDataCommon* AssetLodData = Asset->ClothLodData[0];
-	check(AssetLodData->PhysicalMeshData);
-	UClothPhysicalMeshDataBase* PhysMesh = AssetLodData->PhysicalMeshData;
+	const UClothLODDataCommon* const AssetLodData = Asset->ClothLodData[0];
+	const FClothPhysicalMeshData& PhysMesh = AssetLodData->ClothPhysicalMeshData;
 
 	// SkinPhysicsMesh() strips scale from RootBoneTransform ("Ignore any user scale.
 	// It's already accounted for in our skinning matrices."), and returns all points
@@ -185,13 +189,13 @@ void ClothingSimulation::CreateActor(USkeletalMeshComponent* InOwnerComponent, U
 	TArray<Chaos::TVector<float, 3>> TempAnimationPositions;
 	TArray<Chaos::TVector<float, 3>> TempAnimationNormals;
 
-	TempAnimationPositions.SetNumUninitialized(PhysMesh->Vertices.Num());
-	TempAnimationNormals.SetNumUninitialized(PhysMesh->Vertices.Num());
+	TempAnimationPositions.SetNumUninitialized(PhysMesh.Vertices.Num());
+	TempAnimationNormals.SetNumUninitialized(PhysMesh.Vertices.Num());
 
 	FTransform RootBoneTransform = Context.BoneTransforms[Asset->ReferenceBoneIndex];
 	ClothingMeshUtils::SkinPhysicsMesh<true, false>(
 		Asset->UsedBoneIndices,
-		*PhysMesh, // curr pos and norm
+		PhysMesh, // curr pos and norm
 		Context.ComponentToWorld,
 		Context.RefToLocals.GetData(),
 		Context.RefToLocals.Num(),
@@ -201,12 +205,23 @@ void ClothingSimulation::CreateActor(USkeletalMeshComponent* InOwnerComponent, U
 	// Add particles
 	TPBDParticles<float, 3>& Particles = Evolution->Particles();
 	const uint32 Offset = Particles.Size();
-	Particles.AddParticles(PhysMesh->Vertices.Num());
+	Particles.AddParticles(PhysMesh.Vertices.Num());
 
 	// ClothSharedSimConfig should either be a nullptr, or point to an object common to the whole skeletal mesh
+	UChaosClothSharedSimConfig* const AssetClothConfigShared = Asset->GetClothConfig<UChaosClothSharedSimConfig>();
 	if (ClothSharedSimConfig == nullptr)
 	{
-		ClothSharedSimConfig = Cast<UChaosClothSharedSimConfig>(Asset->ClothSharedSimConfig);
+		ClothSharedSimConfig = AssetClothConfigShared;
+	}
+	else if (!AssetClothConfigShared || AssetClothConfigShared != ClothSharedSimConfig)
+	{
+		UE_CLOG(AssetClothConfigShared,
+			LogChaosCloth, Warning, 
+			TEXT("Found a different Chaos Cloth shared config in the Cloth LOD asset being added to %s in sim slot %d. ")
+			TEXT("The asset's shared config will have to be replaced."),
+			InOwnerComponent->GetOwner() ? *InOwnerComponent->GetOwner()->GetName() : TEXT("None"), InSimDataIndex);
+
+		Asset->SetClothConfig(ClothSharedSimConfig);
 	}
 
 	AnimationPositions.SetNum(Particles.Size());
@@ -228,16 +243,16 @@ void ClothingSimulation::CreateActor(USkeletalMeshComponent* InOwnerComponent, U
 
 	OldAnimationPositions = AnimationPositions;  // Also update the old positions array to avoid any interpolation issues
 
-	const int32 NumTriangles = PhysMesh->Indices.Num() / 3;
+	const int32 NumTriangles = PhysMesh.Indices.Num() / 3;
 	TArray<Chaos::TVector<int32, 3>> InputSurfaceElements;
 	InputSurfaceElements.Reserve(NumTriangles);
 	for (int i = 0; i < NumTriangles; ++i)
 	{
 		const int32 Index = 3 * i;
 		InputSurfaceElements.Add(
-			{ static_cast<int32>(Offset + PhysMesh->Indices[Index]),
-			 static_cast<int32>(Offset + PhysMesh->Indices[Index + 1]),
-			 static_cast<int32>(Offset + PhysMesh->Indices[Index + 2]) });
+			{ static_cast<int32>(Offset + PhysMesh.Indices[Index]),
+			 static_cast<int32>(Offset + PhysMesh.Indices[Index + 1]),
+			 static_cast<int32>(Offset + PhysMesh.Indices[Index + 2]) });
 	}
 	check(InputSurfaceElements.Num() == NumTriangles);
 	if (Meshes.Num() <= InSimDataIndex)
@@ -291,10 +306,11 @@ void ClothingSimulation::CreateActor(USkeletalMeshComponent* InOwnerComponent, U
 		break;
 	};
 	// Clamp and enslave
+	const FPointWeightMap& MaxDistances = PhysMesh.GetWeightMap(EChaosWeightMapTarget::MaxDistance);
 	for (uint32 i = Offset; i < Particles.Size(); i++)
 	{
 		Particles.M(i) = FMath::Max(Particles.M(i), ChaosClothSimConfig->MinPerParticleMass);
-		Particles.InvM(i) = PhysMesh->IsFixed(i - Offset) ? 0.0 : 1.0 / Particles.M(i);
+		Particles.InvM(i) = MaxDistances.IsBelowThreshold(i - Offset) ? 0.f : 1.f / Particles.M(i);
 	}
 
 	// Add Model
@@ -408,12 +424,10 @@ void ClothingSimulation::CreateActor(USkeletalMeshComponent* InOwnerComponent, U
 	}
 
 	// Maximum Distance Constraints
-	const UEnum* const MeshTargets = PhysMesh->GetFloatArrayTargets();
-	const uint32 PhysMeshMaxDistanceIndex = MeshTargets->GetValueByName(TEXT("MaxDistance"));
-	if (PhysMesh->GetFloatArray(PhysMeshMaxDistanceIndex)->Num() > 0)
+	if (MaxDistances.Num() > 0)
 	{
 		check(Mesh->GetNumElements() > 0);
-		Chaos::PBDSphericalConstraint<float, 3> SphericalContraint(Offset, PhysMesh->GetFloatArray(PhysMeshMaxDistanceIndex)->Num(), true, &AnimationPositions, PhysMesh->GetFloatArray(PhysMeshMaxDistanceIndex));
+		Chaos::PBDSphericalConstraint<float, 3> SphericalContraint(Offset, MaxDistances.Num(), true, &AnimationPositions, &MaxDistances.Values);
 		Evolution->AddPBDConstraintFunction([SphericalContraint = MoveTemp(SphericalContraint)](TPBDParticles<float, 3>& InParticles, const float Dt)
 		{
 			SphericalContraint.Apply(InParticles, Dt);
@@ -421,15 +435,15 @@ void ClothingSimulation::CreateActor(USkeletalMeshComponent* InOwnerComponent, U
 	}
 
 	// Backstop Constraints
-	const uint32 PhysMeshBackstopDistanceIndex = MeshTargets->GetValueByName(TEXT("BackstopDistance"));
-	const uint32 PhysMeshBackstopRadiusIndex = MeshTargets->GetValueByName(TEXT("BackstopRadius"));
-	if (PhysMesh->GetFloatArray(PhysMeshBackstopRadiusIndex)->Num() > 0 && PhysMesh->GetFloatArray(PhysMeshBackstopDistanceIndex)->Num() > 0)
+	const FPointWeightMap& BackstopRadiuses = PhysMesh.GetWeightMap(EChaosWeightMapTarget::BackstopRadius);
+	const FPointWeightMap& BackstopDistances = PhysMesh.GetWeightMap(EChaosWeightMapTarget::BackstopDistance);
+	if (BackstopRadiuses.Num() > 0 && BackstopDistances.Num() > 0)
 	{
 		check(Mesh->GetNumElements() > 0);
-		check(PhysMesh->GetFloatArray(PhysMeshBackstopRadiusIndex)->Num() == PhysMesh->GetFloatArray(PhysMeshBackstopDistanceIndex)->Num());
+		check(BackstopRadiuses.Num() == BackstopDistances.Num());
 
-		Chaos::PBDSphericalConstraint<float, 3> SphericalContraint(Offset, PhysMesh->GetFloatArray(PhysMeshBackstopRadiusIndex)->Num(), false, &AnimationPositions,
-			PhysMesh->GetFloatArray(PhysMeshBackstopRadiusIndex), PhysMesh->GetFloatArray(PhysMeshBackstopDistanceIndex), &AnimationNormals);
+		Chaos::PBDSphericalConstraint<float, 3> SphericalContraint(Offset, BackstopRadiuses.Num(), false, &AnimationPositions,
+			&BackstopRadiuses.Values, &BackstopDistances.Values, &AnimationNormals);
 		Evolution->AddPBDConstraintFunction([SphericalContraint = MoveTemp(SphericalContraint)](TPBDParticles<float, 3>& InParticles, const float Dt)
 		{
 			SphericalContraint.Apply(InParticles, Dt);
@@ -437,11 +451,11 @@ void ClothingSimulation::CreateActor(USkeletalMeshComponent* InOwnerComponent, U
 	}
 
 	// Animation Drive Constraints
-	const uint32 PhysMeshAnimDriveIndex = MeshTargets->GetValueByName(TEXT("AnimDriveMultiplier"));
-	if (PhysMesh->GetFloatArray(PhysMeshAnimDriveIndex)->Num() > 0)
+	const FPointWeightMap& AnimDriveMultipliers = PhysMesh.GetWeightMap(EChaosWeightMapTarget::AnimDriveMultiplier);
+	if (AnimDriveMultipliers.Num() > 0)
 	{
 		check(Mesh->GetNumElements() > 0);
-		TPBDAnimDriveConstraint<float, 3> PBDAnimDriveConstraint(Offset, &AnimationPositions, PhysMesh->GetFloatArray(PhysMeshAnimDriveIndex), AnimDriveSpringStiffness[InSimDataIndex]);
+		TPBDAnimDriveConstraint<float, 3> PBDAnimDriveConstraint(Offset, &AnimationPositions, &AnimDriveMultipliers.Values, AnimDriveSpringStiffness[InSimDataIndex]);
 		Evolution->AddPBDConstraintFunction(
 			[PBDAnimDriveConstraint = MoveTemp(PBDAnimDriveConstraint), &Stiffness = AnimDriveSpringStiffness[InSimDataIndex]](TPBDParticles<float, 3>& InParticles, const float Dt) mutable
 		{
@@ -486,19 +500,33 @@ void ClothingSimulation::CreateActor(USkeletalMeshComponent* InOwnerComponent, U
 
 void ClothingSimulation::PostActorCreationInitialize()
 {
-	// Let all assets point to the same shared configuration
-	for (UClothingAssetCommon* Asset : Assets)
+	if (Assets.Num() > 0)
 	{
-		if (Asset)
+		// If we don't have a shared sim config, create one. This is only possible if none of the cloth Assets had a configuration during Actor creation
+		if (ClothSharedSimConfig == nullptr)
+		{
+			// None of the cloth assets had a clothSharedSimConfig, so we will create it
+			ClothSharedSimConfig = NewObject<UChaosClothSharedSimConfig>(Assets[0], UChaosClothSharedSimConfig::StaticClass()->GetFName());
+			check(ClothSharedSimConfig);
+		}
+
+		// Let all assets point to the same shared configuration
+		for (int32 AssetIndex = 0; AssetIndex < Assets.Num(); ++AssetIndex)
 		{
-			// If we don't have a shared sim config, create one. This is only possible if none of the cloth Assets had a configuration during Actor creation
-			if (ClothSharedSimConfig == nullptr)
+			if (UClothingAssetCommon* const Asset = Assets[AssetIndex])
 			{
-				// None of the cloth assets had a clothSharedSimConfig, so we will create it
-				ClothSharedSimConfig = NewObject<UChaosClothSharedSimConfig>(Asset, UChaosClothSharedSimConfig::StaticClass()->GetFName());
-				check(ClothSharedSimConfig);
+				const UChaosClothSharedSimConfig* const AssetClothConfigShared = Asset->GetClothConfig<UChaosClothSharedSimConfig>();
+				if (!AssetClothConfigShared || AssetClothConfigShared != ClothSharedSimConfig)
+				{
+					UE_CLOG(AssetClothConfigShared,
+						LogChaosCloth, Warning,
+						TEXT("Found a different Chaos Cloth shared config in Cloth LOD asset %d. ")
+						TEXT("The asset's shared config will have to be replaced."),
+						AssetIndex);
+
+					Asset->SetClothConfig(ClothSharedSimConfig);
+				}
 			}
-			Asset->ClothSharedSimConfig = ClothSharedSimConfig;
 		}
 	}
 
@@ -510,7 +538,7 @@ void ClothingSimulation::PostActorCreationInitialize()
 		Evolution->SetCollisionThickness(ClothSharedSimConfig->CollisionThickness);
 		Evolution->SetDamping(ClothSharedSimConfig->Damping);
 		Evolution->GetGravityForces().SetAcceleration(Chaos::TVector<float, 3>(ClothSharedSimConfig->Gravity));
-	}	
+	}
 }
 
 void ClothingSimulation::UpdateCollisionTransforms(const ClothingSimulationContext& Context)
@@ -1054,8 +1082,7 @@ void ClothingSimulation::Simulate(IClothingSimulationContext* InContext)
 			}
 
 			const UClothLODDataCommon* AssetLodData = Asset->ClothLodData[0];
-			check(AssetLodData->PhysicalMeshData);
-			const UClothPhysicalMeshDataBase* PhysMesh = AssetLodData->PhysicalMeshData;
+			const FClothPhysicalMeshData& PhysMesh = AssetLodData->ClothPhysicalMeshData;
 			const uint32 PointCount = IndexToRangeMap[Index][1] - IndexToRangeMap[Index][0];
 
 			// Optimization note:
@@ -1067,7 +1094,7 @@ void ClothingSimulation::Simulate(IClothingSimulationContext* InContext)
 
 			ClothingMeshUtils::SkinPhysicsMesh<true, false>(
 				Asset->UsedBoneIndices,
-				*PhysMesh,
+				PhysMesh,
 				Context->ComponentToWorld,
 				Context->RefToLocals.GetData(),
 				Context->RefToLocals.Num(),
@@ -1559,25 +1586,19 @@ void ClothingSimulation::DebugDrawBackstops(USkeletalMeshComponent* OwnerCompone
 		// Get Backstop Distances
 		const UClothLODDataCommon* const AssetLodData = Asset->ClothLodData[0];
 		check(AssetLodData);
-		check(AssetLodData->PhysicalMeshData);
-		const UClothPhysicalMeshDataBase* const PhysMesh = AssetLodData->PhysicalMeshData;
-		const UEnum* const MeshTargets = PhysMesh->GetFloatArrayTargets();
-		const uint32 PhysMeshBackstopIndex = MeshTargets->GetValueByName(TEXT("BackstopDistance"));
-		if (PhysMesh->GetFloatArray(PhysMeshBackstopIndex)->Num() == 0)
-		{
-			continue;
-		}
-
-		const uint32 PhysMeshBackstopRadiusIndex = MeshTargets->GetValueByName(TEXT("BackstopRadius"));
-		if (PhysMesh->GetFloatArray(PhysMeshBackstopRadiusIndex)->Num() == 0)
+		const FClothPhysicalMeshData& PhysMesh = AssetLodData->ClothPhysicalMeshData;
+		const FPointWeightMap& BackstopDistances = PhysMesh.GetWeightMap(EChaosWeightMapTarget::BackstopDistance);
+		const FPointWeightMap& BackstopRadiuses = PhysMesh.GetWeightMap(EChaosWeightMapTarget::BackstopRadius);
+		if (BackstopDistances.Num() == 0 || BackstopRadiuses.Num() == 0)
 		{
 			continue;
 		}
 
 		for (uint32 ParticleIndex = IndexToRangeMap[i][0]; ParticleIndex < IndexToRangeMap[i][1]; ++ParticleIndex)
 		{
-			const float Radius = (*PhysMesh->GetFloatArray(PhysMeshBackstopRadiusIndex))[ParticleIndex - IndexToRangeMap[i][0]];
-			const float Distance = (*PhysMesh->GetFloatArray(PhysMeshBackstopIndex))[ParticleIndex - IndexToRangeMap[i][0]];
+			const uint32 WeightMapIndex = ParticleIndex - IndexToRangeMap[i][0];
+			const float Radius = BackstopRadiuses[WeightMapIndex];
+			const float Distance = BackstopDistances[WeightMapIndex];
 			PDI->DrawLine(AnimationPositions[ParticleIndex], AnimationPositions[ParticleIndex] - AnimationNormals[ParticleIndex] * (Distance - Radius), FColor::White, SDPG_World, 0.0f, 0.001f);
 			if (Radius > 0.0f)
 			{
@@ -1616,18 +1637,17 @@ void ClothingSimulation::DebugDrawMaxDistances(USkeletalMeshComponent* OwnerComp
 		// Get Maximum Distances
 		const UClothLODDataCommon* const AssetLodData = Asset->ClothLodData[0];
 		check(AssetLodData);
-		check(AssetLodData->PhysicalMeshData);
-		UClothPhysicalMeshDataBase* PhysMesh = AssetLodData->PhysicalMeshData;
-		const UEnum* const MeshTargets = PhysMesh->GetFloatArrayTargets();
-		const uint32 PhysMeshMaxDistanceIndex = MeshTargets->GetValueByName(TEXT("MaxDistance"));
-		if (PhysMesh->GetFloatArray(PhysMeshMaxDistanceIndex)->Num() == 0)
+		const FClothPhysicalMeshData& PhysMesh = AssetLodData->ClothPhysicalMeshData;
+		const FPointWeightMap& MaxDistances = PhysMesh.GetWeightMap(EChaosWeightMapTarget::MaxDistance);
+		if (MaxDistances.Num() == 0)
 		{
 			continue;
 		}
 		
 		for (uint32 ParticleIndex = IndexToRangeMap[i][0]; ParticleIndex < IndexToRangeMap[i][1]; ++ParticleIndex)
 		{
-			const float Distance = (*PhysMesh->GetFloatArray(PhysMeshMaxDistanceIndex))[ParticleIndex - IndexToRangeMap[i][0]];
+			const uint32 WeightMapIndex = ParticleIndex - IndexToRangeMap[i][0];
+			const float Distance = MaxDistances[WeightMapIndex];
 			if (Particles.InvM(ParticleIndex) == 0.0f)
 			{
 				const FMatrix& ViewMatrix = PDI->View->ViewMatrices.GetViewMatrix();
@@ -1657,18 +1677,17 @@ void ClothingSimulation::DebugDrawAnimDrive(USkeletalMeshComponent* OwnerCompone
 		// Get Animdrive Multiplier
 		const UClothLODDataCommon* const AssetLodData = Asset->ClothLodData[0];
 		check(AssetLodData);
-		check(AssetLodData->PhysicalMeshData);
-		const UClothPhysicalMeshDataBase* const PhysMesh = AssetLodData->PhysicalMeshData;
-		const UEnum* const MeshTargets = PhysMesh->GetFloatArrayTargets();
-		const uint32 PhysMeshAnimDriveIndex = MeshTargets->GetValueByName(TEXT("AnimDriveMultiplier"));
-		if (PhysMesh->GetFloatArray(PhysMeshAnimDriveIndex)->Num() == 0)
+		const FClothPhysicalMeshData& PhysMesh = AssetLodData->ClothPhysicalMeshData;
+		const FPointWeightMap& AnimDriveMultipliers = PhysMesh.GetWeightMap(EChaosWeightMapTarget::AnimDriveMultiplier);
+		if (AnimDriveMultipliers.Num() == 0)
 		{
 			continue;
 		}
 
 		for (uint32 ParticleIndex = IndexToRangeMap[i][0]; ParticleIndex < IndexToRangeMap[i][1]; ++ParticleIndex)
 		{
-			const float Multiplier = (*PhysMesh->GetFloatArray(PhysMeshAnimDriveIndex))[ParticleIndex - IndexToRangeMap[i][0]];
+			const uint32 WeightMapIndex = ParticleIndex - IndexToRangeMap[i][0];
+			const float Multiplier = AnimDriveMultipliers[WeightMapIndex];
 			PDI->DrawLine(AnimationPositions[ParticleIndex], Particles.X(ParticleIndex), Multiplier * AnimDriveSpringStiffness[i] * FColor::Cyan, SDPG_World, 0.0f, 0.001f);
 		}
 	}
diff --git a/Engine/Plugins/Experimental/ChaosCloth/Source/Chaos/Private/ChaosCloth/ChaosClothingSimulationFactory.cpp b/Engine/Plugins/Experimental/ChaosCloth/Source/Chaos/Private/ChaosCloth/ChaosClothingSimulationFactory.cpp
--- a/Engine/Plugins/Experimental/ChaosCloth/Source/Chaos/Private/ChaosCloth/ChaosClothingSimulationFactory.cpp
+++ b/Engine/Plugins/Experimental/ChaosCloth/Source/Chaos/Private/ChaosCloth/ChaosClothingSimulationFactory.cpp
@@ -3,6 +3,7 @@
 #include "ChaosCloth/ChaosClothingSimulationFactory.h"
 #include "ChaosCloth/ChaosClothConfig.h"
 #include "ChaosCloth/ChaosClothingSimulation.h"
+#include "ChaosCloth/ChaosWeightMapTarget.h"
 #include "ChaosCloth/ChaosClothingSimulationInteractor.h"
 
 IClothingSimulation* UChaosClothingSimulationFactory::CreateSimulation()
@@ -31,7 +32,12 @@ UClothingSimulationInteractor* UChaosClothingSimulationFactory::CreateInteractor
 	return NewObject<UChaosClothingSimulationInteractor>(GetTransientPackage());
 }
 
-UClothConfigBase* UChaosClothingSimulationFactory::CreateDefaultClothConfig(const FObjectInitializer& ObjectInitializer, UObject* Outer)
+TSubclassOf<UClothConfigBase> UChaosClothingSimulationFactory::GetClothConfigClass() const
 {
-	return ObjectInitializer.CreateDefaultSubobject<UChaosClothConfig>(Outer, UChaosClothConfig::StaticClass()->GetFName());
-}
\ No newline at end of file
+	return TSubclassOf<UClothConfigBase>(UChaosClothConfig::StaticClass());
+}
+
+const UEnum* UChaosClothingSimulationFactory::GetWeightMapTargetEnum() const
+{
+	return StaticEnum<EChaosWeightMapTarget>();
+}
diff --git a/Engine/Plugins/Experimental/ChaosCloth/Source/Chaos/Private/ChaosCloth/ChaosWeightMapTarget.h b/Engine/Plugins/Experimental/ChaosCloth/Source/Chaos/Private/ChaosCloth/ChaosWeightMapTarget.h
new file mode 100644
--- /dev/null
+++ b/Engine/Plugins/Experimental/ChaosCloth/Source/Chaos/Private/ChaosCloth/ChaosWeightMapTarget.h
@@ -0,0 +1,18 @@
+// Copyright 1998-2019 Epic Games, Inc. All Rights Reserved.
+#pragma once
+
+#include "ClothPhysicalMeshData.h"
+#include "ChaosWeightMapTarget.generated.h"
+
+
+/** Targets for painted weight maps (aka masks). */
+UENUM()
+enum class EChaosWeightMapTarget : uint8
+{
+	None                = EWeightMapTargetCommon::None,
+	MaxDistance         = EWeightMapTargetCommon::MaxDistance,
+	BackstopDistance    = EWeightMapTargetCommon::BackstopDistance,
+	BackstopRadius      = EWeightMapTargetCommon::BackstopRadius,
+	AnimDriveMultiplier = EWeightMapTargetCommon::AnimDriveMultiplier
+	// Add Chaos specific maps below this line
+};
diff --git a/Engine/Plugins/Experimental/ChaosCloth/Source/Chaos/Public/ChaosCloth/ChaosClothConfig.h b/Engine/Plugins/Experimental/ChaosCloth/Source/Chaos/Public/ChaosCloth/ChaosClothConfig.h
--- a/Engine/Plugins/Experimental/ChaosCloth/Source/Chaos/Public/ChaosCloth/ChaosClothConfig.h
+++ b/Engine/Plugins/Experimental/ChaosCloth/Source/Chaos/Public/ChaosCloth/ChaosClothConfig.h
@@ -2,18 +2,18 @@
 #pragma once
 #include "Components/SkeletalMeshComponent.h"
 #include "Engine/SkeletalMesh.h"
-#include "ClothConfigBase.h"
+#include "ClothConfig.h"
 #include "CoreMinimal.h"
 #include "ChaosClothConfig.generated.h"
 
 /** Holds initial, asset level config for clothing actors. */
 // Hiding categories that will be used in the future
 UCLASS(HideCategories = (Collision))
-class CHAOSCLOTH_API UChaosClothConfig : public UClothConfigBase
+class CHAOSCLOTH_API UChaosClothConfig : public UClothConfigCommon
 {
 	GENERATED_BODY()
 public:
-	UChaosClothConfig() {};
+	UChaosClothConfig() {}
 
 	/**
 	 * How cloth particle mass is determined
@@ -99,16 +99,19 @@ public:
 	bool bUseContinuousCollisionDetection = false;
 };
 
-/*
-These settings are shared between all instances on a skeletal mesh
-*/
+/**
+ * Chaos config settings shared between all instances of a skeletal mesh.
+ * Unlike UChaosClothConfig, these settings contain common cloth simulation
+ * parameters that cannot change between the various clothing assets assigned
+ * to a specific skeletal mesh. @seealso UChaosClothConfig.
+ */
 UCLASS()
-class CHAOSCLOTH_API UChaosClothSharedSimConfig : public UClothSharedSimConfigBase
+class CHAOSCLOTH_API UChaosClothSharedSimConfig : public UClothConfigBase
 {
 	GENERATED_BODY()
 public:
-	UChaosClothSharedSimConfig() : Gravity(FVector(0.0f, 0.0f, -490)) {};
-	virtual ~UChaosClothSharedSimConfig() {};
+	UChaosClothSharedSimConfig() : Gravity(FVector(0.0f, 0.0f, -490)) {}
+	virtual ~UChaosClothSharedSimConfig() {}
 
 	// The number of solver iterations
 	// This will increase the stiffness of all constraints but will increase the CPU cost
diff --git a/Engine/Plugins/Experimental/ChaosCloth/Source/Chaos/Public/ChaosCloth/ChaosClothingSimulationFactory.h b/Engine/Plugins/Experimental/ChaosCloth/Source/Chaos/Public/ChaosCloth/ChaosClothingSimulationFactory.h
--- a/Engine/Plugins/Experimental/ChaosCloth/Source/Chaos/Public/ChaosCloth/ChaosClothingSimulationFactory.h
+++ b/Engine/Plugins/Experimental/ChaosCloth/Source/Chaos/Public/ChaosCloth/ChaosClothingSimulationFactory.h
@@ -7,14 +7,17 @@
 #include "ChaosClothingSimulationFactory.generated.h"
 
 UCLASS()
-class CHAOSCLOTH_API UChaosClothingSimulationFactory : public UClothingSimulationFactory
+class CHAOSCLOTH_API UChaosClothingSimulationFactory final : public UClothingSimulationFactory
 {
-    GENERATED_BODY()
-  public:
-    virtual IClothingSimulation* CreateSimulation() override;
-    virtual void DestroySimulation(IClothingSimulation* InSimulation) override;
-    virtual bool SupportsAsset(UClothingAssetBase* InAsset) override;
-    virtual bool SupportsRuntimeInteraction() override;
+	GENERATED_BODY()
+public:
+	virtual IClothingSimulation* CreateSimulation() override;
+	virtual void DestroySimulation(IClothingSimulation* InSimulation) override;
+	virtual bool SupportsAsset(UClothingAssetBase* InAsset) override;
+
+	virtual bool SupportsRuntimeInteraction() override;
 	virtual UClothingSimulationInteractor* CreateInteractor() override;
-	virtual UClothConfigBase* CreateDefaultClothConfig(const FObjectInitializer& ObjectInitializer, UObject* Outer) override;
+
+	virtual TSubclassOf<UClothConfigBase> GetClothConfigClass() const override;
+	const UEnum* GetWeightMapTargetEnum() const override;
 };
diff --git a/Engine/Source/Editor/ClothPainter/Private/ClothMeshAdapter.cpp b/Engine/Source/Editor/ClothPainter/Private/ClothMeshAdapter.cpp
--- a/Engine/Source/Editor/ClothPainter/Private/ClothMeshAdapter.cpp
+++ b/Engine/Source/Editor/ClothPainter/Private/ClothMeshAdapter.cpp
@@ -234,11 +234,10 @@ bool FClothMeshPaintAdapter::InitializeVertexData()
 			{
 				UClothingAssetCommon* ConcreteAsset = CastChecked<UClothingAssetCommon>(SelectedAsset);
 				const UClothLODDataCommon* LODData = ConcreteAsset->ClothLodData[PaintingClothLODIndex];
-				check(LODData->PhysicalMeshData);
-				const UClothPhysicalMeshDataBase* MeshData = LODData->PhysicalMeshData;
+				const FClothPhysicalMeshData& MeshData = LODData->ClothPhysicalMeshData;
 
 				MeshVertices.Append(DebugComponent->SkinnedSelectedClothingPositions);
-				MeshIndices.Append(MeshData->Indices);
+				MeshIndices.Append(MeshData.Indices);
 
 				for(int32 Index = IndexOffset; Index < MeshIndices.Num(); ++Index)
 				{
@@ -249,8 +248,8 @@ bool FClothMeshPaintAdapter::InitializeVertexData()
 				Info.IndexStart = IndexOffset;
 				Info.VertexStart = VertexOffset;
 
-				IndexOffset += MeshData->Indices.Num();
-				VertexOffset += MeshData->Vertices.Num();
+				IndexOffset += MeshData.Indices.Num();
+				VertexOffset += MeshData.Vertices.Num();
 
 				Info.IndexEnd = IndexOffset;
 				Info.VertexEnd = VertexOffset;
diff --git a/Engine/Source/Editor/ClothPainter/Private/SClothAssetSelector.cpp b/Engine/Source/Editor/ClothPainter/Private/SClothAssetSelector.cpp
--- a/Engine/Source/Editor/ClothPainter/Private/SClothAssetSelector.cpp
+++ b/Engine/Source/Editor/ClothPainter/Private/SClothAssetSelector.cpp
@@ -8,7 +8,6 @@
 #include "Widgets/Input/SButton.h"
 #include "Widgets/Images/SImage.h"
 #include "ClothingAsset.h"
-#include "ClothPhysicalMeshDataBase.h"
 #include "Framework/MultiBox/MultiBoxBuilder.h"
 #include "Widgets/Input/SCheckBox.h"
 #include "Widgets/Text/SInlineEditableTextBlock.h"
@@ -50,10 +49,10 @@ FPointWeightMap* FClothingMaskListItem::GetMask()
 	return nullptr;
 }
 
-UClothPhysicalMeshDataBase* FClothingMaskListItem::GetMeshData()
+FClothPhysicalMeshData* FClothingMaskListItem::GetMeshData()
 {
 	UClothingAssetCommon * Asset = ClothingAsset.Get();
-	return Asset && Asset->IsValidLod(LodIndex) ? Asset->ClothLodData[LodIndex]->PhysicalMeshData : nullptr;
+	return Asset && Asset->IsValidLod(LodIndex) ? &Asset->ClothLodData[LodIndex]->ClothPhysicalMeshData : nullptr;
 }
 
 USkeletalMesh* FClothingMaskListItem::GetOwningMesh()
@@ -319,14 +318,14 @@ private:
 					DestMask.CurrentTarget = SourceMask.CurrentTarget;
 
 					ClothingMeshUtils::FVertexParameterMapper ParameterMapper(
-						DestLod->PhysicalMeshData->Vertices,
-						DestLod->PhysicalMeshData->Normals,
-						SourceLod->PhysicalMeshData->Vertices,
-						SourceLod->PhysicalMeshData->Normals,
-						SourceLod->PhysicalMeshData->Indices
+						DestLod->ClothPhysicalMeshData.Vertices,
+						DestLod->ClothPhysicalMeshData.Normals,
+						SourceLod->ClothPhysicalMeshData.Vertices,
+						SourceLod->ClothPhysicalMeshData.Normals,
+						SourceLod->ClothPhysicalMeshData.Indices
 					);
 
-					ParameterMapper.Map(SourceMask.GetValueArray(), DestMask.Values);
+					ParameterMapper.Map(SourceMask.Values, DestMask.Values);
 				}
 			}
 		}
@@ -404,11 +403,16 @@ public:
 
 		if(InColumnName == Column_CurrentTarget)
 		{
-			FPointWeightMap* Mask = Item->GetMask();
-			UEnum* Enum = Item->GetMeshData()->GetFloatArrayTargets();
-			if(Enum && Mask)
+			// Retrieve the mask names for the current clothing simulation factory
+			const TSubclassOf<class UClothingSimulationFactory> ClothingSimulationFactory = UClothingSimulationFactory::GetDefaultClothingSimulationFactoryClass();
+			if (ClothingSimulationFactory.Get() != nullptr)
 			{
-				return SNew(STextBlock).Text(Enum->GetDisplayNameTextByIndex((int32)Mask->CurrentTarget));
+				const UEnum* const Enum = ClothingSimulationFactory.GetDefaultObject()->GetWeightMapTargetEnum();
+				const FPointWeightMap* const Mask = Item->GetMask();
+				if(Enum && Mask)
+				{
+					return SNew(STextBlock).Text(Enum->GetDisplayNameTextByIndex((int32)Mask->CurrentTarget));
+				}
 			}
 		}
 
@@ -551,7 +555,7 @@ private:
 			if(FPointWeightMap* Mask = Item->GetMask())
 			{
 				Mask->CurrentTarget = (uint8)InTargetEntryIndex;
-				if(Mask->CurrentTarget == 0)//(uint8)MaskTarget_PhysMesh::None)
+				if(Mask->CurrentTarget == (uint8)EWeightMapTargetCommon::None)
 				{
 					// Make sure to disable this mask if it has no valid target
 					Mask->bEnabled = false;
@@ -566,20 +570,23 @@ private:
 	{
 		Builder.BeginSection(NAME_None, LOCTEXT("MaskTargets_SectionName", "Targets"));
 		{
-			//UEnum* Enum = StaticEnum<MaskTarget_PhysMesh>();
-			UEnum* Enum = Item->GetMeshData()->GetFloatArrayTargets();
-			if(Enum)
+			// Retrieve the mask names for the current clothing simulation factory
+			const TSubclassOf<class UClothingSimulationFactory> ClothingSimulationFactory = UClothingSimulationFactory::GetDefaultClothingSimulationFactoryClass();
+			if (ClothingSimulationFactory.Get() != nullptr)
 			{
-				const int32 NumEntries = Enum->NumEnums();
-
-				// Iterate to -1 to skip the _MAX entry appended to the end of the enum
-				for(int32 Index = 0; Index < NumEntries - 1; ++Index)
+				if (const UEnum* const Enum = ClothingSimulationFactory.GetDefaultObject()->GetWeightMapTargetEnum())
 				{
-					FUIAction EntryAction(FExecuteAction::CreateSP(this, &SMaskListRow::OnSetTarget, Index));
+					const int32 NumEntries = Enum->NumEnums();
 
-					FText EntryText = Enum->GetDisplayNameTextByIndex(Index);
+					// Iterate to -1 to skip the _MAX entry appended to the end of the enum
+					for(int32 Index = 0; Index < NumEntries - 1; ++Index)
+					{
+						FUIAction EntryAction(FExecuteAction::CreateSP(this, &SMaskListRow::OnSetTarget, Index));
+
+						FText EntryText = Enum->GetDisplayNameTextByIndex(Index);
 
-					Builder.AddMenuEntry(EntryText, FText::GetEmpty(), FSlateIcon(), EntryAction);
+						Builder.AddMenuEntry(EntryText, FText::GetEmpty(), FSlateIcon(), EntryAction);
+					}
 				}
 			}
 		}
@@ -621,7 +628,7 @@ private:
 		{
 			if(FPointWeightMap* Mask = InItem->GetMask())
 			{
-				return Mask->CurrentTarget != 0; // (uint8)MaskTarget_PhysMesh::None;
+				return Mask->CurrentTarget != (uint8)EWeightMapTargetCommon::None;
 			}
 		}
 
@@ -1130,14 +1137,14 @@ FReply SClothAssetSelector::AddNewMask()
 		if(Asset->ClothLodData.IsValidIndex(SelectedLod))
 		{
 			UClothLODDataCommon* LodData = Asset->ClothLodData[SelectedLod];
-			const int32 NumRequiredValues = LodData->PhysicalMeshData->Vertices.Num();
+			const int32 NumRequiredValues = LodData->ClothPhysicalMeshData.Vertices.Num();
 
 			LodData->ParameterMasks.AddDefaulted();
 
 			FPointWeightMap& NewMask = LodData->ParameterMasks.Last();
 
 			NewMask.Name = TEXT("New Mask");
-			NewMask.CurrentTarget = 0;// (uint8)MaskTarget_PhysMesh::None;
+			NewMask.CurrentTarget = (uint8)EWeightMapTargetCommon::None;
 			NewMask.Values.AddZeroed(NumRequiredValues);
 
 			OnRefresh();
diff --git a/Engine/Source/Editor/ClothPainter/Private/SClothAssetSelector.h b/Engine/Source/Editor/ClothPainter/Private/SClothAssetSelector.h
--- a/Engine/Source/Editor/ClothPainter/Private/SClothAssetSelector.h
+++ b/Engine/Source/Editor/ClothPainter/Private/SClothAssetSelector.h
@@ -11,7 +11,7 @@
 struct FAssetData;
 class USkeletalMesh;
 class UClothingAssetCommon;
-class UClothPhysicalMeshDataBase;
+struct FClothPhysicalMeshData;
 struct FPointWeightMap;
 
 struct FClothingAssetListItem
@@ -27,7 +27,7 @@ struct FClothingMaskListItem
 	{}
 
 	FPointWeightMap* GetMask();
-	UClothPhysicalMeshDataBase* GetMeshData();
+	FClothPhysicalMeshData* GetMeshData();
 	USkeletalMesh* GetOwningMesh();
 
 	TWeakObjectPtr<UClothingAssetCommon> ClothingAsset;
diff --git a/Engine/Source/Editor/ClothPainter/Private/SCopyVertexColorSettingsPanel.cpp b/Engine/Source/Editor/ClothPainter/Private/SCopyVertexColorSettingsPanel.cpp
--- a/Engine/Source/Editor/ClothPainter/Private/SCopyVertexColorSettingsPanel.cpp
+++ b/Engine/Source/Editor/ClothPainter/Private/SCopyVertexColorSettingsPanel.cpp
@@ -114,13 +114,13 @@ FReply SCopyVertexColorSettingsPanel::OnCopyClicked()
 	{
 		UClothLODDataCommon* ClothLODData = Asset->ClothLodData[SelectedLOD];
 
-		check(ClothLODData->PhysicalMeshData->Vertices.Num() == ClothLODData->PhysicalMeshData->VertexColors.Num());
-		int32 NumVerts = ClothLODData->PhysicalMeshData->Vertices.Num();
+		check(ClothLODData->ClothPhysicalMeshData.Vertices.Num() == ClothLODData->ClothPhysicalMeshData.VertexColors.Num());
+		int32 NumVerts = ClothLODData->ClothPhysicalMeshData.Vertices.Num();
 		check(SelectedMask->Values.Num() == NumVerts);
 
 		for (int32 VertIdx = 0; VertIdx < NumVerts; VertIdx++)
 		{
-			const FColor VertColor = ClothLODData->PhysicalMeshData->VertexColors[VertIdx];
+			const FColor VertColor = ClothLODData->ClothPhysicalMeshData.VertexColors[VertIdx];
 			SelectedMask->Values[VertIdx] = GetColorChannelAsFloat(VertColor, CopyParams.ColorChannel, CopyParams.ScalingFactor);
 		}
 	}
diff --git a/Engine/Source/Editor/ClothingSystemEditor/Private/ClothingAssetFactory.cpp b/Engine/Source/Editor/ClothingSystemEditor/Private/ClothingAssetFactory.cpp
--- a/Engine/Source/Editor/ClothingSystemEditor/Private/ClothingAssetFactory.cpp
+++ b/Engine/Source/Editor/ClothingSystemEditor/Private/ClothingAssetFactory.cpp
@@ -6,11 +6,10 @@
 #include "PhysXPublic.h"
 #include "PhysicsPublic.h"
 #include "ClothingAsset.h"
-#include "ClothPhysicalMeshDataNv.h" // for MaskTarget_PhysMesh
 #if WITH_APEX_CLOTHING
 #include "ClothingAssetAuthoring.h"
-#include "Assets/ClothingAssetNv.h"
 #include "ApexClothingUtils.h"
+#include "ClothConfigNv.h"
 #endif // WITH_APEX_CLOTHING
 #include "ContentBrowserModule.h"
 #include "Misc/FileHelper.h"
@@ -25,6 +24,7 @@
 #include "Framework/Notifications/NotificationManager.h"
 #include "Utils/ClothingMeshUtils.h"
 #include "Rendering/SkeletalMeshModel.h"
+#include "ClothPhysicalMeshData.h"
 
 #define LOCTEXT_NAMESPACE "ClothingAssetFactory"
 DEFINE_LOG_CATEGORY(LogClothingAssetFactory)
@@ -238,12 +238,8 @@ UClothingAssetBase* UClothingAssetFactory::CreateFromSkeletalMesh(USkeletalMesh*
 
 	FString SanitizedName = ObjectTools::SanitizeObjectName(Params.AssetName);
 	FName ObjectName = MakeUniqueObjectName(TargetMesh, UClothingAssetCommon::StaticClass(), FName(*SanitizedName));
-#if WITH_APEX_CLOTHING
-	UClothingAssetCommon* NewAsset = NewObject<UClothingAssetNv>(TargetMesh, ObjectName);
-#else
-	check(false); // TODO: Implement Chaos ClothingAsset, and have the factory get the correct object
+
 	UClothingAssetCommon* NewAsset = NewObject<UClothingAssetCommon>(TargetMesh, ObjectName);
-#endif
 	NewAsset->SetFlags(RF_Transactional);
 
 	// Adding a new LOD from this skeletal mesh
@@ -389,7 +385,7 @@ UClothingAssetBase* UClothingAssetFactory::ImportLodToClothing(USkeletalMesh* Ta
 UClothingAssetBase* UClothingAssetFactory::CreateFromApexAsset(nvidia::apex::ClothingAsset* InApexAsset, USkeletalMesh* TargetMesh, FName InName)
 {
 #if WITH_APEX_CLOTHING
-	UClothingAssetNv* NewClothingAsset = NewObject<UClothingAssetNv>(TargetMesh, InName);
+	UClothingAssetCommon* NewClothingAsset = NewObject<UClothingAssetCommon>(TargetMesh, InName);
 	NewClothingAsset->SetFlags(RF_Transactional);
 
 	const NvParameterized::Interface* AssetParams = InApexAsset->getAssetNvParameterized();
@@ -417,19 +413,19 @@ UClothingAssetBase* UClothingAssetFactory::CreateFromApexAsset(nvidia::apex::Clo
 		ExtractMaterialParameters(NewClothingAsset, *InApexAsset);
 
 		// Set to use legacy wind calculations, which is what APEX would normally have used
-		UClothConfigNv* ConfigNv = Cast<UClothConfigNv>(NewClothingAsset->ClothSimConfig);
+		UClothConfigNv* const ConfigNv = NewClothingAsset->GetClothConfig<UClothConfigNv>();
 		if (!ensure(ConfigNv))
 		{
 			continue;
 		}
-		ConfigNv->WindMethod = EClothingWindMethod::Legacy;
+		ConfigNv->ClothingWindMethod = EClothingWindMethodNv::Legacy;
 
 		// Fixup unreal-side bone indices
-		const int32 NumBoneDatas = CurrentLodData.PhysicalMeshData->BoneData.Num();
+		const int32 NumBoneDatas = CurrentLodData.ClothPhysicalMeshData.BoneData.Num();
 		check(NumBoneDatas == ApexVertData.Num());
 		for(int32 BoneDataIndex = 0; BoneDataIndex < NumBoneDatas; ++BoneDataIndex)
 		{
-			FClothVertBoneData& BoneData = CurrentLodData.PhysicalMeshData->BoneData[BoneDataIndex];
+			FClothVertBoneData& BoneData = CurrentLodData.ClothPhysicalMeshData.BoneData[BoneDataIndex];
 			FApexVertData& CurrentVertData = ApexVertData[BoneDataIndex];
 		
 			for(int32 BoneInfluenceIdx = 0; BoneInfluenceIdx < MAX_TOTAL_INFLUENCES; ++BoneInfluenceIdx)
@@ -479,36 +475,31 @@ UClothingAssetBase* UClothingAssetFactory::CreateFromApexAsset(nvidia::apex::Clo
 	for (UClothLODDataCommon* LodPtr : NewClothingAsset->ClothLodData)
 	{
 		check(LodPtr);
-		UClothPhysicalMeshDataNv* PhysMeshPtr = Cast<UClothPhysicalMeshDataNv>(LodPtr->PhysicalMeshData);
-		if (!ensure(PhysMeshPtr))
-		{
-			continue;
-		}
-		UClothPhysicalMeshDataNv& PhysMesh = *PhysMeshPtr;
+		const FClothPhysicalMeshData& PhysMesh = LodPtr->ClothPhysicalMeshData;
 
-		// Didn't do anything previously - clear out incase there's something in there
+		// Didn't do anything previously - clear out in case there's something in there
 		// so we can use it correctly now.
 		LodPtr->ParameterMasks.Reset(3);
 
 		// Max distances
 		LodPtr->ParameterMasks.AddDefaulted();
 		FPointWeightMap& MaxDistanceMask = LodPtr->ParameterMasks.Last();
-		MaxDistanceMask.CopyFrom(PhysMesh.GetFloatArray((int32)MaskTarget_PhysMesh::MaxDistance), (uint8)MaskTarget_PhysMesh::MaxDistance);
-		MaxDistanceMask.bEnabled = true;
+		const FPointWeightMap& MaxDistances = PhysMesh.GetWeightMap(EWeightMapTargetCommon::MaxDistance);
+		MaxDistanceMask.Initialize(MaxDistances, EWeightMapTargetCommon::MaxDistance);
 
-		if(PhysMesh.BackstopRadiuses.FindByPredicate([](const float& A) {return A != 0.0f; }))
+		const FPointWeightMap* const BackstopRadiuses = PhysMesh.FindWeightMap(EWeightMapTargetCommon::BackstopRadius);
+		if (BackstopRadiuses && !BackstopRadiuses->IsZeroed())
 		{
 			// Backstop radii
 			LodPtr->ParameterMasks.AddDefaulted();
 			FPointWeightMap& BackstopRadiusMask = LodPtr->ParameterMasks.Last();
-			BackstopRadiusMask.CopyFrom(PhysMesh.GetFloatArray((int32)MaskTarget_PhysMesh::BackstopRadius), (uint8)MaskTarget_PhysMesh::BackstopRadius);
-			BackstopRadiusMask.bEnabled = true;
+			BackstopRadiusMask.Initialize(*BackstopRadiuses, EWeightMapTargetCommon::BackstopRadius);
 
 			// Backstop distances
 			LodPtr->ParameterMasks.AddDefaulted();
 			FPointWeightMap& BackstopDistanceMask = LodPtr->ParameterMasks.Last();
-			BackstopDistanceMask.CopyFrom(PhysMesh.GetFloatArray((int32)MaskTarget_PhysMesh::BackstopDistance), (uint8)MaskTarget_PhysMesh::BackstopDistance);
-			BackstopDistanceMask.bEnabled = true;
+			const FPointWeightMap& BackstopDistances = PhysMesh.GetWeightMap(EWeightMapTargetCommon::BackstopDistance);
+			BackstopDistanceMask.Initialize(BackstopDistances, EWeightMapTargetCommon::BackstopDistance);
 		}
 	}
 
@@ -1008,7 +999,7 @@ void UClothingAssetFactory::ExtractMaterialParameters(UClothingAssetCommon* NewA
 
 	MaterialArrayHandle.set(MaterialIndex);
 
-	if(UClothConfigNv* ConfigPtr = Cast<UClothConfigNv>(NewAsset->ClothSimConfig))
+	if(UClothConfigNv* const ConfigPtr = NewAsset->GetClothConfig<UClothConfigNv>())
 	{
 		UClothConfigNv& Config = *ConfigPtr;
 		NvParameterized::Handle ChildHandle(MaterialArrayHandle);
@@ -1062,71 +1053,71 @@ void UClothingAssetFactory::ExtractMaterialParameters(UClothingAssetCommon* NewA
 		// Vertical constraint params
 
 		MaterialArrayHandle.getChildHandle(MaterialArrayHandle.getInterface(), "verticalStretchingStiffness", ChildHandle);
-		ChildHandle.getParamF32(Config.VerticalConstraintConfig.Stiffness);
+		ChildHandle.getParamF32(Config.VerticalConstraint.Stiffness);
 
 		MaterialArrayHandle.getChildHandle(MaterialArrayHandle.getInterface(), "verticalStiffnessScaling.compressionRange", ChildHandle);
-		ChildHandle.getParamF32(Config.VerticalConstraintConfig.CompressionLimit);
+		ChildHandle.getParamF32(Config.VerticalConstraint.CompressionLimit);
 
 		MaterialArrayHandle.getChildHandle(MaterialArrayHandle.getInterface(), "verticalStiffnessScaling.stretchRange", ChildHandle);
-		ChildHandle.getParamF32(Config.VerticalConstraintConfig.StretchLimit);
+		ChildHandle.getParamF32(Config.VerticalConstraint.StretchLimit);
 
 		MaterialArrayHandle.getChildHandle(MaterialArrayHandle.getInterface(), "verticalStiffnessScaling.scale", ChildHandle);
-		ChildHandle.getParamF32(Config.VerticalConstraintConfig.StiffnessMultiplier);
+		ChildHandle.getParamF32(Config.VerticalConstraint.StiffnessMultiplier);
 
 		// Horizontal constraint params
 
 		MaterialArrayHandle.getChildHandle(MaterialArrayHandle.getInterface(), "horizontalStretchingStiffness", ChildHandle);
-		ChildHandle.getParamF32(Config.HorizontalConstraintConfig.Stiffness);
+		ChildHandle.getParamF32(Config.HorizontalConstraint.Stiffness);
 
 		MaterialArrayHandle.getChildHandle(MaterialArrayHandle.getInterface(), "horizontalStiffnessScaling.compressionRange", ChildHandle);
-		ChildHandle.getParamF32(Config.HorizontalConstraintConfig.CompressionLimit);
+		ChildHandle.getParamF32(Config.HorizontalConstraint.CompressionLimit);
 
 		MaterialArrayHandle.getChildHandle(MaterialArrayHandle.getInterface(), "horizontalStiffnessScaling.stretchRange", ChildHandle);
-		ChildHandle.getParamF32(Config.HorizontalConstraintConfig.StretchLimit);
+		ChildHandle.getParamF32(Config.HorizontalConstraint.StretchLimit);
 
 		MaterialArrayHandle.getChildHandle(MaterialArrayHandle.getInterface(), "horizontalStiffnessScaling.scale", ChildHandle);
-		ChildHandle.getParamF32(Config.HorizontalConstraintConfig.StiffnessMultiplier);
+		ChildHandle.getParamF32(Config.HorizontalConstraint.StiffnessMultiplier);
 
 		// Bend constraint params
 
 		MaterialArrayHandle.getChildHandle(MaterialArrayHandle.getInterface(), "bendingStiffness", ChildHandle);
-		ChildHandle.getParamF32(Config.BendConstraintConfig.Stiffness);
+		ChildHandle.getParamF32(Config.BendConstraint.Stiffness);
 
 		MaterialArrayHandle.getChildHandle(MaterialArrayHandle.getInterface(), "bendingStiffnessScaling.compressionRange", ChildHandle);
-		ChildHandle.getParamF32(Config.BendConstraintConfig.CompressionLimit);
+		ChildHandle.getParamF32(Config.BendConstraint.CompressionLimit);
 
 		MaterialArrayHandle.getChildHandle(MaterialArrayHandle.getInterface(), "bendingStiffnessScaling.stretchRange", ChildHandle);
-		ChildHandle.getParamF32(Config.BendConstraintConfig.StretchLimit);
+		ChildHandle.getParamF32(Config.BendConstraint.StretchLimit);
 
 		MaterialArrayHandle.getChildHandle(MaterialArrayHandle.getInterface(), "bendingStiffnessScaling.scale", ChildHandle);
-		ChildHandle.getParamF32(Config.BendConstraintConfig.StiffnessMultiplier);
+		ChildHandle.getParamF32(Config.BendConstraint.StiffnessMultiplier);
 
 		// Shear constraint params
 
 		MaterialArrayHandle.getChildHandle(MaterialArrayHandle.getInterface(), "shearingStiffness", ChildHandle);
-		ChildHandle.getParamF32(Config.ShearConstraintConfig.Stiffness);
+		ChildHandle.getParamF32(Config.ShearConstraint.Stiffness);
 
 		MaterialArrayHandle.getChildHandle(MaterialArrayHandle.getInterface(), "shearingStiffnessScaling.compressionRange", ChildHandle);
-		ChildHandle.getParamF32(Config.ShearConstraintConfig.CompressionLimit);
+		ChildHandle.getParamF32(Config.ShearConstraint.CompressionLimit);
 
 		MaterialArrayHandle.getChildHandle(MaterialArrayHandle.getInterface(), "shearingStiffnessScaling.stretchRange", ChildHandle);
-		ChildHandle.getParamF32(Config.ShearConstraintConfig.StretchLimit);
+		ChildHandle.getParamF32(Config.ShearConstraint.StretchLimit);
 
 		MaterialArrayHandle.getChildHandle(MaterialArrayHandle.getInterface(), "shearingStiffnessScaling.scale", ChildHandle);
-		ChildHandle.getParamF32(Config.ShearConstraintConfig.StiffnessMultiplier);
+		ChildHandle.getParamF32(Config.ShearConstraint.StiffnessMultiplier);
 
 		// UE just used the vertical config for everything, so stomp the other configs
-		Config.HorizontalConstraintConfig.CompressionLimit = Config.VerticalConstraintConfig.CompressionLimit;
-		Config.HorizontalConstraintConfig.StretchLimit = Config.VerticalConstraintConfig.StretchLimit;
-		Config.HorizontalConstraintConfig.StiffnessMultiplier = Config.VerticalConstraintConfig.StiffnessMultiplier;
+		Config.HorizontalConstraint.CompressionLimit = Config.VerticalConstraint.CompressionLimit;
+		Config.HorizontalConstraint.StretchLimit = Config.VerticalConstraint.StretchLimit;
+		Config.HorizontalConstraint.StiffnessMultiplier = Config.VerticalConstraint.StiffnessMultiplier;
 
-		Config.BendConstraintConfig.CompressionLimit = Config.VerticalConstraintConfig.CompressionLimit;
-		Config.BendConstraintConfig.StretchLimit = Config.VerticalConstraintConfig.StretchLimit;
-		Config.BendConstraintConfig.StiffnessMultiplier = Config.VerticalConstraintConfig.StiffnessMultiplier;
+		Config.BendConstraint.CompressionLimit = Config.VerticalConstraint.CompressionLimit;
+		Config.BendConstraint.StretchLimit = Config.VerticalConstraint.StretchLimit;
+		Config.BendConstraint.StiffnessMultiplier = Config.VerticalConstraint.StiffnessMultiplier;
 
-		Config.ShearConstraintConfig.CompressionLimit = Config.VerticalConstraintConfig.CompressionLimit;
-		Config.ShearConstraintConfig.StretchLimit = Config.VerticalConstraintConfig.StretchLimit;
-		Config.ShearConstraintConfig.StiffnessMultiplier = Config.VerticalConstraintConfig.StiffnessMultiplier;
+		Config.ShearConstraint.CompressionLimit = Config.VerticalConstraint.CompressionLimit;
+		Config.ShearConstraint.StretchLimit = Config.VerticalConstraint.StretchLimit;
+		Config.ShearConstraint.StiffnessMultiplier = Config.VerticalConstraint.StiffnessMultiplier;
 
 	}
 }
@@ -1230,8 +1221,7 @@ bool UClothingAssetFactory::ImportToLodInternal(
 
 	if(InParameterRemapSource)
 	{
-		check(InParameterRemapSource->PhysicalMeshData);
-		UClothPhysicalMeshDataBase& RemapPhysMesh = *InParameterRemapSource->PhysicalMeshData;
+		const FClothPhysicalMeshData& RemapPhysMesh = InParameterRemapSource->ClothPhysicalMeshData;
 		CachedPositions = RemapPhysMesh.Vertices;
 		CachedNormals = RemapPhysMesh.Normals;
 		CachedIndices = RemapPhysMesh.Indices;
@@ -1241,13 +1231,8 @@ bool UClothingAssetFactory::ImportToLodInternal(
 		bPerformParamterRemap = true;
 	}
 
-	check(DestLod->PhysicalMeshData);
-	UClothPhysicalMeshDataBase& PhysMesh = *DestLod->PhysicalMeshData;
-	PhysMesh.Reset(NumUniqueVerts);
-	PhysMesh.Indices.Reset(NumIndices);
-	PhysMesh.Indices.AddZeroed(NumIndices);
-
-	TArray<TArray<float>*> WeightArrays = PhysMesh.GetFloatArrays();
+	FClothPhysicalMeshData& PhysMesh = DestLod->ClothPhysicalMeshData;
+	PhysMesh.Reset(NumUniqueVerts, NumIndices);
 
 	for(int32 VertexIndex = 0; VertexIndex < NumUniqueVerts; ++VertexIndex)
 	{
@@ -1257,15 +1242,6 @@ bool UClothingAssetFactory::ImportToLodInternal(
 		PhysMesh.Normals[VertexIndex] = SourceVert.TangentZ;
 		PhysMesh.VertexColors[VertexIndex] = SourceVert.Color;
 
-		for (int32 WeightArrayIndex = 0; WeightArrayIndex < WeightArrays.Num(); ++WeightArrayIndex)
-		{
-			TArray<float>* Weights = WeightArrays[WeightArrayIndex];
-			if (Weights && Weights->IsValidIndex(VertexIndex))
-			{
-				(*Weights)[VertexIndex] = 0.0f;
-			}
-		}
-
 		FClothVertBoneData& BoneData = PhysMesh.BoneData[VertexIndex];
 		for(int32 InfluenceIndex = 0; InfluenceIndex < MAX_TOTAL_INFLUENCES; ++InfluenceIndex)
 		{
@@ -1282,8 +1258,8 @@ bool UClothingAssetFactory::ImportToLodInternal(
 	// Add a max distance parameter mask to begin with
 	DestLod->ParameterMasks.AddDefaulted();
 	FPointWeightMap& Mask = DestLod->ParameterMasks.Last();
-	Mask.CopyFrom(PhysMesh.GetFloatArray((uint32)MaskTarget_PhysMesh::MaxDistance), (uint8)MaskTarget_PhysMesh::MaxDistance);
-	Mask.bEnabled = true;
+	const FPointWeightMap& MaxDistances = PhysMesh.GetWeightMap(EWeightMapTargetCommon::MaxDistance);
+	Mask.Initialize(MaxDistances, EWeightMapTargetCommon::MaxDistance);
 
 	PhysMesh.MaxBoneWeights = SourceSection.MaxBoneInfluences;
 
@@ -1319,7 +1295,7 @@ bool UClothingAssetFactory::ImportToLodInternal(
 
 		for(int32 MaskIndex = 0; MaskIndex < NumSourceMasks; ++MaskIndex)
 		{
-			FPointWeightMap& SourceMask = SourceMaskCopy[MaskIndex];
+			const FPointWeightMap& SourceMask = SourceMaskCopy[MaskIndex];
 
 			DestLod->ParameterMasks.AddDefaulted();
 			FPointWeightMap& DestMask = DestLod->ParameterMasks.Last();
@@ -1328,7 +1304,7 @@ bool UClothingAssetFactory::ImportToLodInternal(
 			DestMask.CurrentTarget = SourceMask.CurrentTarget;
 			DestMask.bEnabled = SourceMask.bEnabled;
 
-			ParameterRemapper.Map(SourceMask.GetValueArray(), DestMask.Values);
+			ParameterRemapper.Map(SourceMask.Values, DestMask.Values);
 		}
 
 		DestAsset->ApplyParameterMasks();
@@ -1343,13 +1319,7 @@ void UClothingAssetFactory::ExtractLodPhysicalData(UClothingAssetCommon* NewAsse
 {
 	const NvParameterized::Interface* AssetParams = InApexAsset.getAssetNvParameterized();
 
-	UClothPhysicalMeshDataNv* PhysDataPtr = Cast<UClothPhysicalMeshDataNv>(InLodData.PhysicalMeshData);
-	if (!ensure(PhysDataPtr))
-	{
-		return;
-	}
-	UClothPhysicalMeshDataNv& PhysData = *PhysDataPtr;
-
+	FClothPhysicalMeshData& PhysData = InLodData.ClothPhysicalMeshData;
 	NvParameterized::Handle GraphicalMeshArrayHandle(*AssetParams, "graphicalLods");
 
 	int32 NumGraphicalLods = 0;
@@ -1518,9 +1488,15 @@ void UClothingAssetFactory::ExtractLodPhysicalData(UClothingAssetCommon* NewAsse
 
 		check(CoeffArraySize == NumVertices);
 
-		PhysData.MaxDistances.AddZeroed(CoeffArraySize);
-		PhysData.BackstopDistances.AddZeroed(CoeffArraySize);
-		PhysData.BackstopRadiuses.AddZeroed(CoeffArraySize);
+		TArray<float>& MaxDistances = PhysData.FindOrAddWeightMap(EWeightMapTargetCommon::MaxDistance).Values;
+		TArray<float>& BackstopDistances = PhysData.FindOrAddWeightMap(EWeightMapTargetCommon::BackstopDistance).Values;
+		TArray<float>& BackstopRadiuses = PhysData.FindOrAddWeightMap(EWeightMapTargetCommon::BackstopRadius).Values;
+		MaxDistances.Reset(CoeffArraySize);
+		MaxDistances.AddZeroed(CoeffArraySize);
+		BackstopDistances.Reset(CoeffArraySize);
+		BackstopDistances.AddZeroed(CoeffArraySize);
+		BackstopRadiuses.Reset(CoeffArraySize);
+		BackstopRadiuses.AddZeroed(CoeffArraySize);
 
 		IterHandle = NvParameterized::Handle(TempHandle);
 		NvParameterized::Handle ChildHandle(PhysicalMeshRef);
@@ -1529,20 +1505,20 @@ void UClothingAssetFactory::ExtractLodPhysicalData(UClothingAssetCommon* NewAsse
 			TempHandle.getChildHandle(Idx, IterHandle);
 
 			IterHandle.getChildHandle(PhysicalMeshRef, "maxDistance", ChildHandle);
-			ChildHandle.getParamF32(PhysData.MaxDistances[Idx]);
+			ChildHandle.getParamF32(MaxDistances[Idx]);
 
 			IterHandle.getChildHandle(AssetParams, "collisionSphereDistance", ChildHandle);
-			ChildHandle.getParamF32(PhysData.BackstopDistances[Idx]);
+			ChildHandle.getParamF32(BackstopDistances[Idx]);
 
 			IterHandle.getChildHandle(AssetParams, "collisionSphereRadius", ChildHandle);
-			ChildHandle.getParamF32(PhysData.BackstopRadiuses[Idx]);
+			ChildHandle.getParamF32(BackstopRadiuses[Idx]);
 
-			PhysData.BackstopDistances[Idx] += PhysData.BackstopRadiuses[Idx];
+			BackstopDistances[Idx] += BackstopRadiuses[Idx];
 		}
 
 		// Calculate how many fixed verts we have
 		PhysData.NumFixedVerts = 0;
-		for(float Distance : PhysData.MaxDistances)
+		for(float Distance : MaxDistances)
 		{
 			if(Distance == 0.0f)
 			{
diff --git a/Engine/Source/Editor/ClothingSystemEditor/Public/ClothingAssetFactory.h b/Engine/Source/Editor/ClothingSystemEditor/Public/ClothingAssetFactory.h
--- a/Engine/Source/Editor/ClothingSystemEditor/Public/ClothingAssetFactory.h
+++ b/Engine/Source/Editor/ClothingSystemEditor/Public/ClothingAssetFactory.h
@@ -4,7 +4,6 @@
 
 #include "ClothingAsset.h"
 #include "ClothingAssetFactoryInterface.h"
-#include "ClothLODDataNv.h"
 #include "GPUSkinPublicDefs.h"
 
 #include "ClothingAssetFactory.generated.h"
@@ -13,6 +12,7 @@
 DECLARE_LOG_CATEGORY_EXTERN(LogClothingAssetFactory, Log, All);
 
 class FSkeletalMeshLODModel;
+class UClothingAssetCommon;
 
 namespace nvidia
 {
diff --git a/Engine/Source/Editor/Persona/Private/PersonaMeshDetails.cpp b/Engine/Source/Editor/Persona/Private/PersonaMeshDetails.cpp
--- a/Engine/Source/Editor/Persona/Private/PersonaMeshDetails.cpp
+++ b/Engine/Source/Editor/Persona/Private/PersonaMeshDetails.cpp
@@ -5155,10 +5155,9 @@ TSharedRef<SUniformGridPanel> FPersonaMeshDetails::MakeClothingDetailsWidget(int
 
 		if (UClothingAssetCommon* Asset = Cast<UClothingAssetCommon>(ClothingAsset))
 		{
-			UClothLODDataCommon* LodData = Asset->ClothLodData[LODIndex];
-			check(LodData->PhysicalMeshData);
-			UClothPhysicalMeshDataBase& PhysMeshData = *LodData->PhysicalMeshData;
-			FClothCollisionData& CollisionData = LodData->CollisionData;
+			const UClothLODDataCommon* const LodData = Asset->ClothLodData[LODIndex];
+			const FClothPhysicalMeshData& PhysMeshData = LodData->ClothPhysicalMeshData;
+			const FClothCollisionData& CollisionData = LodData->CollisionData;
 
 			Grid->AddSlot(0, RowNumber)
 				.HAlign(HAlign_Center)
diff --git a/Engine/Source/Editor/UnrealEd/Private/Animation/DebugSkelMeshComponent.cpp b/Engine/Source/Editor/UnrealEd/Private/Animation/DebugSkelMeshComponent.cpp
--- a/Engine/Source/Editor/UnrealEd/Private/Animation/DebugSkelMeshComponent.cpp
+++ b/Engine/Source/Editor/UnrealEd/Private/Animation/DebugSkelMeshComponent.cpp
@@ -742,10 +742,11 @@ void UDebugSkelMeshComponent::RebuildClothingSectionsFixedVerts()
 				{
 					UClothingAssetCommon* ConcreteAsset = Cast<UClothingAssetCommon>(BaseAsset);
 					UClothLODDataCommon* LodData = ConcreteAsset->ClothLodData[Section.ClothingData.AssetLodIndex];
+					const FPointWeightMap& MaxDistances = LodData->ClothPhysicalMeshData.GetWeightMap(EWeightMapTargetCommon::MaxDistance);
 
 					for(FMeshToMeshVertData& VertData : Section.ClothMappingData)
 					{
-						if(LodData->PhysicalMeshData->IsFixed(
+						if(MaxDistances.AreAnyBelowThreshold(
 							VertData.SourceMeshVertIndices[0],
 							VertData.SourceMeshVertIndices[1],
 							VertData.SourceMeshVertIndices[2]))
@@ -830,7 +831,7 @@ void UDebugSkelMeshComponent::RefreshSelectedClothingSkinnedPositions()
 
 				UClothLODDataCommon* LodData = ConcreteAsset->ClothLodData[SelectedClothingLodForPainting];
 
-				ClothingMeshUtils::SkinPhysicsMesh(ConcreteAsset->UsedBoneIndices, *LodData->PhysicalMeshData, FTransform::Identity, RefToLocals.GetData(), RefToLocals.Num(), SkinnedSelectedClothingPositions, SkinnedSelectedClothingNormals);
+				ClothingMeshUtils::SkinPhysicsMesh(ConcreteAsset->UsedBoneIndices, LodData->ClothPhysicalMeshData, FTransform::Identity, RefToLocals.GetData(), RefToLocals.Num(), SkinnedSelectedClothingPositions, SkinnedSelectedClothingNormals);
 				RebuildCachedClothBounds();
 			}
 		}
@@ -1037,13 +1038,13 @@ FDebugSkelMeshDynamicData::FDebugSkelMeshDynamicData(UDebugSkelMeshComponent* In
 						{
 							UClothLODDataCommon* LodData = ConcreteAsset->ClothLodData[InComponent->SelectedClothingLodForPainting];
 
-							ClothingSimIndices = LodData->PhysicalMeshData->Indices;
+							ClothingSimIndices = LodData->ClothPhysicalMeshData.Indices;
 
 							if(LodData->ParameterMasks.IsValidIndex(InComponent->SelectedClothingLodMaskForPainting))
 							{
 								FPointWeightMap& Mask = LodData->ParameterMasks[InComponent->SelectedClothingLodMaskForPainting];
 
-								ClothingVisiblePropertyValues = Mask.GetValueArray();
+								ClothingVisiblePropertyValues = Mask.Values;
 							}
 						}
 					}
diff --git a/Engine/Source/Runtime/ClothingSystemRuntimeCommon/Private/ClothConfig.cpp b/Engine/Source/Runtime/ClothingSystemRuntimeCommon/Private/ClothConfig.cpp
new file mode 100644
--- /dev/null
+++ b/Engine/Source/Runtime/ClothingSystemRuntimeCommon/Private/ClothConfig.cpp
@@ -0,0 +1,8 @@
+// Copyright 1998-2019 Epic Games, Inc. All Rights Reserved.
+#include "ClothConfig.h"
+
+UClothConfigCommon::UClothConfigCommon()
+{}
+
+UClothConfigCommon::~UClothConfigCommon()
+{}
diff --git a/Engine/Source/Runtime/ClothingSystemRuntimeCommon/Private/ClothLODData.cpp b/Engine/Source/Runtime/ClothingSystemRuntimeCommon/Private/ClothLODData.cpp
--- a/Engine/Source/Runtime/ClothingSystemRuntimeCommon/Private/ClothLODData.cpp
+++ b/Engine/Source/Runtime/ClothingSystemRuntimeCommon/Private/ClothLODData.cpp
@@ -1,11 +1,13 @@
 // Copyright 1998-2019 Epic Games, Inc. All Rights Reserved.
 
 #include "ClothLODData.h"
-
+#include "ClothPhysicalMeshData.h"
 
 UClothLODDataCommon::UClothLODDataCommon(const FObjectInitializer& Init)
-	: PhysicalMeshData(nullptr)
-{}
+	: Super(Init)
+	, PhysicalMeshData_DEPRECATED(nullptr)
+{
+}
 
 UClothLODDataCommon::~UClothLODDataCommon()
 {}
@@ -45,21 +47,27 @@ void UClothLODDataCommon::Serialize(FArchive& Ar)
 	   << TransitionDownSkinData;
 }
 
+void UClothLODDataCommon::PostLoad()
+{
+	Super::PostLoad();
+
+	if (PhysicalMeshData_DEPRECATED)
+	{
+		ClothPhysicalMeshData.MigrateFrom(PhysicalMeshData_DEPRECATED);
+		PhysicalMeshData_DEPRECATED = nullptr;
+	}
+}
+
 #if WITH_EDITOR
 void UClothLODDataCommon::PushWeightsToMesh()
 {
-	if (PhysicalMeshData)
+	ClothPhysicalMeshData.ClearWeightMaps();
+	for (const FPointWeightMap& Weights : ParameterMasks)
 	{
-		PhysicalMeshData->ClearParticleParameters();
-		for (const FPointWeightMap &Weights : ParameterMasks)
+		if (Weights.bEnabled)
 		{
-			if (Weights.bEnabled)
-			{
-				if (TArray<float>* TargetArray = PhysicalMeshData->GetFloatArray(Weights.CurrentTarget))
-				{
-					*TargetArray = Weights.GetValueArray();
-				}
-			}
+			FPointWeightMap& TargetWeightMap = ClothPhysicalMeshData.FindOrAddWeightMap(Weights.CurrentTarget);
+			TargetWeightMap.Values = Weights.Values;
 		}
 	}
 }
diff --git a/Engine/Source/Runtime/ClothingSystemRuntimeCommon/Private/ClothLODData_Legacy.cpp b/Engine/Source/Runtime/ClothingSystemRuntimeCommon/Private/ClothLODData_Legacy.cpp
new file mode 100644
--- /dev/null
+++ b/Engine/Source/Runtime/ClothingSystemRuntimeCommon/Private/ClothLODData_Legacy.cpp
@@ -0,0 +1,60 @@
+// Copyright 1998-2019 Epic Games, Inc. All Rights Reserved.
+
+#include "ClothLODData_Legacy.h"
+#include "PointWeightMap.h"
+#include "ClothLODData.h"
+
+FClothParameterMask_Legacy::FClothParameterMask_Legacy()
+	: MaskName(NAME_None)
+	, CurrentTarget(EWeightMapTargetCommon::None)
+	, MaxValue_DEPRECATED(0.0)
+	, MinValue_DEPRECATED(100.0)
+	, bEnabled(false)
+{}
+
+void FClothParameterMask_Legacy::MigrateTo(FPointWeightMap* Weights) const
+{
+	Weights->Name = MaskName;
+	Weights->CurrentTarget = static_cast<uint8>(CurrentTarget);
+	Weights->Values = Values;
+	Weights->bEnabled = bEnabled;
+}
+
+bool FClothLODData_Legacy::Serialize(FArchive& Ar)
+{
+	// Serialize normal tagged data
+	if (!Ar.IsCountingMemory())
+	{
+		UScriptStruct* Struct = FClothLODData_Legacy::StaticStruct();
+		Struct->SerializeTaggedProperties(Ar, (uint8*)this, Struct, nullptr);
+	}
+	// Serialize the mesh to mesh data (not a USTRUCT)
+	Ar	<< TransitionUpSkinData
+		<< TransitionDownSkinData;
+	return true;
+}
+
+void FClothLODData_Legacy::MigrateTo(UClothLODDataCommon* LodData) const
+{
+	PhysicalMeshData.MigrateTo(LodData->ClothPhysicalMeshData);
+	LodData->CollisionData = CollisionData;
+#if WITH_CHAOS
+	// Rebuild surface points so that the legacy Apex convex collision data can also be used with Chaos
+	for (auto& Convex : LodData->CollisionData.Convexes)
+	{
+		if (!Convex.SurfacePoints.Num())
+		{
+			Convex.RebuildSurfacePoints();
+		}
+	}
+#endif // #if WITH_CHAOS
+#if WITH_EDITORONLY_DATA
+	LodData->ParameterMasks.SetNum(ParameterMasks.Num());
+	for (int i = 0; i < ParameterMasks.Num(); i++)
+	{
+		ParameterMasks[i].MigrateTo(&LodData->ParameterMasks[i]);
+	}
+#endif // WITH_EDITORONLY_DATA
+	LodData->TransitionUpSkinData = TransitionUpSkinData;
+	LodData->TransitionDownSkinData = TransitionDownSkinData;
+}
diff --git a/Engine/Source/Runtime/ClothingSystemRuntimeCommon/Private/ClothPhysicalMeshData.cpp b/Engine/Source/Runtime/ClothingSystemRuntimeCommon/Private/ClothPhysicalMeshData.cpp
new file mode 100644
--- /dev/null
+++ b/Engine/Source/Runtime/ClothingSystemRuntimeCommon/Private/ClothPhysicalMeshData.cpp
@@ -0,0 +1,178 @@
+// Copyright 1998-2019 Epic Games, Inc. All Rights Reserved.
+#include "ClothPhysicalMeshData.h"
+#include "ClothConfigBase.h"
+#include "ClothPhysicalMeshDataBase_Legacy.h"
+
+FClothPhysicalMeshData::FClothPhysicalMeshData()
+	: MaxBoneWeights(0)
+	, NumFixedVerts(0)
+{
+	ClearWeightMaps();
+}
+
+void FClothPhysicalMeshData::MigrateTo(FClothPhysicalMeshData& ClothPhysicalMeshData) const
+{
+	// Migrate properties
+	ClothPhysicalMeshData.Vertices = Vertices;
+	ClothPhysicalMeshData.Normals = Normals;
+#if WITH_EDITORONLY_DATA
+	ClothPhysicalMeshData.VertexColors = VertexColors;
+#endif
+	ClothPhysicalMeshData.Indices = Indices;
+	ClothPhysicalMeshData.InverseMasses = InverseMasses;
+	ClothPhysicalMeshData.BoneData = BoneData;
+	ClothPhysicalMeshData.NumFixedVerts = NumFixedVerts;
+	ClothPhysicalMeshData.MaxBoneWeights = MaxBoneWeights;
+	ClothPhysicalMeshData.SelfCollisionIndices = SelfCollisionIndices;
+
+	// Migrate maps
+	if (MaxDistances_DEPRECATED.Num())
+	{
+		ClothPhysicalMeshData.GetWeightMap(EWeightMapTargetCommon::MaxDistance).Values = MaxDistances_DEPRECATED;
+	}
+	if (BackstopDistances_DEPRECATED.Num())
+	{
+		ClothPhysicalMeshData.GetWeightMap(EWeightMapTargetCommon::BackstopDistance).Values = BackstopDistances_DEPRECATED;
+	}
+	if (BackstopRadiuses_DEPRECATED.Num())
+	{
+		ClothPhysicalMeshData.GetWeightMap(EWeightMapTargetCommon::BackstopRadius).Values = BackstopRadiuses_DEPRECATED;
+	}
+	if (AnimDriveMultipliers_DEPRECATED.Num())
+	{
+		ClothPhysicalMeshData.GetWeightMap(EWeightMapTargetCommon::AnimDriveMultiplier).Values = AnimDriveMultipliers_DEPRECATED;
+	}
+}
+
+void FClothPhysicalMeshData::MigrateFrom(const UClothPhysicalMeshDataBase_Legacy* ClothPhysicalMeshDataBase)
+{
+	Vertices = ClothPhysicalMeshDataBase->Vertices;
+	Normals = ClothPhysicalMeshDataBase->Normals;
+#if WITH_EDITORONLY_DATA
+	VertexColors = ClothPhysicalMeshDataBase->VertexColors;
+#endif
+	Indices = ClothPhysicalMeshDataBase->Indices;
+	InverseMasses = ClothPhysicalMeshDataBase->InverseMasses;
+	BoneData = ClothPhysicalMeshDataBase->BoneData;
+	NumFixedVerts = ClothPhysicalMeshDataBase->NumFixedVerts;
+	MaxBoneWeights = ClothPhysicalMeshDataBase->MaxBoneWeights;
+	SelfCollisionIndices = ClothPhysicalMeshDataBase->SelfCollisionIndices;
+
+	const TArray<uint32> FloatArrayIds = ClothPhysicalMeshDataBase->GetFloatArrayIds();
+	for (uint32 FloatArrayId : FloatArrayIds)
+	{
+		if (const TArray<float>* const FloatArray = ClothPhysicalMeshDataBase->GetFloatArray(FloatArrayId))
+		{
+			FindOrAddWeightMap(FloatArrayId).Values = *FloatArray;
+		}
+	}
+}
+
+void FClothPhysicalMeshData::Reset(const int32 InNumVerts, const int32 InNumIndices)
+{
+	Vertices.Init(FVector::ZeroVector, InNumVerts);
+	Normals.Init(FVector::ZeroVector, InNumVerts);
+#if WITH_EDITORONLY_DATA
+	VertexColors.Init(FColor::Black, InNumVerts);
+#endif //#if WITH_EDITORONLY_DATA
+	InverseMasses.Init(0.f, InNumVerts);
+	BoneData.Reset(InNumVerts);
+	BoneData.AddDefaulted(InNumVerts);
+	Indices.Init(0, InNumIndices);
+
+	NumFixedVerts = 0;
+	MaxBoneWeights = 0;
+
+	ClearWeightMaps();
+}
+
+void FClothPhysicalMeshData::ClearWeightMaps()
+{
+	// Clear all weight maps (and reserve 4 slots)
+	WeightMaps.Empty(4);
+
+	// Max distances must be present, so fill to zero on clear so we still have valid mesh data.
+	AddWeightMap(EWeightMapTargetCommon::MaxDistance).Initialize(Vertices.Num());
+
+	// Add default (empty) optional maps, as these are always expected to be found
+	AddWeightMap(EWeightMapTargetCommon::BackstopDistance);
+	AddWeightMap(EWeightMapTargetCommon::BackstopRadius);
+	AddWeightMap(EWeightMapTargetCommon::AnimDriveMultiplier);
+}
+
+void FClothPhysicalMeshData::BuildSelfCollisionData(const TMap<FName, UClothConfigBase*>& ClothConfigs)
+{
+#if WITH_APEX_CLOTHING  // Only apex clothing needs to build the SelfCollisionIndices
+	float SCRadius = 0.f;
+	for (const TPair<FName, UClothConfigBase*>& ClothConfig : ClothConfigs)
+	{
+		SCRadius = ClothConfig.Value->NeedsSelfCollisionIndices();
+		if (SCRadius > 0.f)
+		{
+			break; 
+		}
+	}
+
+	if (SCRadius == 0.f) { return; }
+
+	const float SCRadiusSq = SCRadius * SCRadius;
+
+	// Start with the full set
+	const int32 NumVerts = Vertices.Num();
+	SelfCollisionIndices.Reset();
+	const FPointWeightMap& MaxDistances = GetWeightMap(EWeightMapTargetCommon::MaxDistance);
+	for (int32 Index = 0; Index < NumVerts; ++Index)
+	{
+		if (!MaxDistances.IsBelowThreshold(Index))
+		{
+			SelfCollisionIndices.Add(Index);
+		}
+	}
+
+	// Now start aggressively culling verts that are near others that we have accepted
+	for (int32 Vert0Itr = 0; Vert0Itr < SelfCollisionIndices.Num(); ++Vert0Itr)
+	{
+		const uint32 V0Index = SelfCollisionIndices[Vert0Itr];
+		if (V0Index == INDEX_NONE)
+		{
+			// We'll remove these indices later.  Just skip it for now.
+			continue;
+		}
+
+		const FVector& V0Pos = Vertices[V0Index];
+
+		// Start one after our current V0, we've done the other checks
+		for (int32 Vert1Itr = Vert0Itr + 1; Vert1Itr < SelfCollisionIndices.Num(); ++Vert1Itr)
+		{
+			const uint32 V1Index = SelfCollisionIndices[Vert1Itr];
+			if (V1Index == INDEX_NONE)
+			{
+				// We'll remove these indices later.  Just skip it for now.
+				continue;
+			}
+
+			const FVector& V1Pos = Vertices[V1Index];
+			const float V0ToV1DistSq = (V1Pos - V0Pos).SizeSquared();
+			if (V0ToV1DistSq < SCRadiusSq)
+			{
+				// Points are in contact in the rest state.  Remove it.
+				//
+				// It's worth noting that this biases towards removing indices 
+				// of later in the list, and keeping ones earlier.  That's not
+				// a great criteria for choosing which one is more important.
+				SelfCollisionIndices[Vert1Itr] = INDEX_NONE;
+				continue;
+			}
+		}
+	}
+
+	// Cull flagged indices.
+	for (int32 It = SelfCollisionIndices.Num(); It--;)
+	{
+		if (SelfCollisionIndices[It] == INDEX_NONE)
+		{
+			SelfCollisionIndices.RemoveAt(It);
+		}
+	}
+#endif  // WITH_APEX_CLOTHING
+}
diff --git a/Engine/Source/Runtime/ClothingSystemRuntimeCommon/Private/ClothingAsset.cpp b/Engine/Source/Runtime/ClothingSystemRuntimeCommon/Private/ClothingAsset.cpp
--- a/Engine/Source/Runtime/ClothingSystemRuntimeCommon/Private/ClothingAsset.cpp
+++ b/Engine/Source/Runtime/ClothingSystemRuntimeCommon/Private/ClothingAsset.cpp
@@ -1,6 +1,9 @@
 // Copyright 1998-2019 Epic Games, Inc. All Rights Reserved.
 
 #include "ClothingAsset.h"
+#include "ClothingAssetCustomVersion.h"
+#include "ClothPhysicalMeshData.h"
+#include "ClothConfig.h"
 
 #include "Utils/ClothingMeshUtils.h"
 
@@ -172,11 +175,36 @@ void ClothingAssetUtils::ClearSectionClothingData(FSkelMeshSection& InSection)
 UClothingAssetCommon::UClothingAssetCommon(const FObjectInitializer& ObjectInitializer)
 	: Super(ObjectInitializer)
 	, PhysicsAsset(nullptr)
-	, ClothSimConfig(nullptr)
-	, ChaosClothSimConfig(nullptr)
+	, ClothSimConfig_DEPRECATED(nullptr)
+	, ChaosClothSimConfig_DEPRECATED(nullptr)
 	, ReferenceBoneIndex(0)
 	, CustomData(nullptr)
-{}
+{
+	// Iterate through all the current providers to create the required sim configs
+	const TArray<IClothingSimulationFactoryClassProvider*> ClassProviders = IModularFeatures::Get().GetModularFeatureImplementations<IClothingSimulationFactoryClassProvider>(IClothingSimulationFactoryClassProvider::FeatureName);
+	ClothConfigs.Reserve(ClassProviders.Num());
+
+	// Create all necessary sim configs (one simulation config per provider)
+	for (int32 Index = 0; Index < ClassProviders.Num(); ++Index)
+	{
+		IClothingSimulationFactoryClassProvider* const Provider = ClassProviders[Index];
+		check(Provider);
+		if (UClass* const ClothingSimulationFactory = *TSubclassOf<class UClothingSimulationFactory>(Provider->GetClothingSimulationFactoryClass()))
+		{
+			const TSubclassOf<UClothConfigBase> ClothConfigClass = ClothingSimulationFactory->GetDefaultObject<UClothingSimulationFactory>()->GetClothConfigClass();
+
+			UClothConfigBase* const ClothConfig = static_cast<UClothConfigBase*>(
+				ObjectInitializer.CreateDefaultSubobject(
+					this,
+					ClothConfigClass->GetFName(),
+					UClothConfigBase::StaticClass(),
+					ClothConfigClass,
+					/*bIsRequired =*/ true,
+					/*bTransient =*/ false));
+			SetClothConfig(ClothConfig);
+		}
+	}
+}
 
 #if WITH_EDITOR
 
@@ -293,9 +321,9 @@ bool UClothingAssetCommon::BindToSkeletalMesh(
 
 	ClothingMeshUtils::ClothMeshDesc TargetMesh(RenderPositions, RenderNormals, IndexView);
 	ClothingMeshUtils::ClothMeshDesc SourceMesh(
-		LodData->PhysicalMeshData->Vertices, 
-		LodData->PhysicalMeshData->Normals, 
-		LodData->PhysicalMeshData->Indices);
+		LodData->ClothPhysicalMeshData.Vertices, 
+		LodData->ClothPhysicalMeshData.Normals, 
+		LodData->ClothPhysicalMeshData.Indices);
 
 	ClothingMeshUtils::GenerateMeshToMeshSkinningData(MeshToMeshData, TargetMesh, &RenderTangents, SourceMesh);
 
@@ -307,9 +335,10 @@ bool UClothingAssetCommon::BindToSkeletalMesh(
 	}
 
 	// Calculate fixed verts
+	const FPointWeightMap& MaxDistances = LodData->ClothPhysicalMeshData.GetWeightMap(EWeightMapTargetCommon::MaxDistance);
 	for(FMeshToMeshVertData& VertData : MeshToMeshData)
 	{
-		if(LodData->PhysicalMeshData->IsFixed(
+		if(MaxDistances.AreAnyBelowThreshold(
 			VertData.SourceMeshVertIndices[0], 
 			VertData.SourceMeshVertIndices[1], 
 			VertData.SourceMeshVertIndices[2])) // Default threshold is 0.1, not 0.0.  Using 0.1 for consistency.
@@ -521,8 +550,7 @@ void UClothingAssetCommon::BuildLodTransitionData()
 		const bool bHasNextLod = LodIndex < NumLods - 1;
 
 		UClothLODDataCommon* CurrentLod = ClothLodData[LodIndex];
-		check(CurrentLod->PhysicalMeshData);
-		UClothPhysicalMeshDataBase& CurrentPhysMesh = *CurrentLod->PhysicalMeshData;
+		const FClothPhysicalMeshData& CurrentPhysMesh = CurrentLod->ClothPhysicalMeshData;
 
 		UClothLODDataCommon* PrevLod = bHasPrevLod ? ClothLodData[LodIndex - 1] : nullptr;
 		UClothLODDataCommon* NextLod = bHasNextLod ? ClothLodData[LodIndex + 1] : nullptr;
@@ -533,16 +561,14 @@ void UClothingAssetCommon::BuildLodTransitionData()
 
 		if(PrevLod)
 		{
-			check(PrevLod->PhysicalMeshData);
-			UClothPhysicalMeshDataBase& PrevPhysMesh = *PrevLod->PhysicalMeshData;
+			FClothPhysicalMeshData& PrevPhysMesh = PrevLod->ClothPhysicalMeshData;
 			CurrentLod->TransitionUpSkinData.Empty(CurrentLodNumVerts);
 			ClothingMeshUtils::ClothMeshDesc PrevMeshDesc(PrevPhysMesh.Vertices, PrevPhysMesh.Normals, PrevPhysMesh.Indices);
 			ClothingMeshUtils::GenerateMeshToMeshSkinningData(CurrentLod->TransitionUpSkinData, CurrentMeshDesc, nullptr, PrevMeshDesc);
 		}
 		if(NextLod)
 		{
-			check(NextLod->PhysicalMeshData);
-			UClothPhysicalMeshDataBase& NextPhysMesh = *NextLod->PhysicalMeshData;
+			FClothPhysicalMeshData& NextPhysMesh = NextLod->ClothPhysicalMeshData;
 			CurrentLod->TransitionDownSkinData.Empty(CurrentLodNumVerts);
 			ClothingMeshUtils::ClothMeshDesc NextMeshDesc(NextPhysMesh.Vertices, NextPhysMesh.Normals, NextPhysMesh.Indices);
 			ClothingMeshUtils::GenerateMeshToMeshSkinningData(CurrentLod->TransitionDownSkinData, CurrentMeshDesc, nullptr, NextMeshDesc);
@@ -599,9 +625,9 @@ void UClothingAssetCommon::CalculateReferenceBoneIndex()
 
 		for(UClothLODDataCommon* CurLod : ClothLodData)
 		{
-			check(CurLod && CurLod->PhysicalMeshData);
-			UClothPhysicalMeshDataBase* MeshData = CurLod->PhysicalMeshData;
-			for(FClothVertBoneData& VertBoneData : MeshData->BoneData)
+			check(CurLod);
+			const FClothPhysicalMeshData& MeshData = CurLod->ClothPhysicalMeshData;
+			for(const FClothVertBoneData& VertBoneData : MeshData.BoneData)
 			{
 				for(int32 InfluenceIndex = 0; InfluenceIndex < MAX_TOTAL_INFLUENCES; ++InfluenceIndex)
 				{
@@ -695,17 +721,288 @@ int32 UClothingAssetCommon::GetNumLods()
 
 void UClothingAssetCommon::BuildSelfCollisionData()
 {
-	if(!ClothSimConfig || !ClothSimConfig->HasSelfCollision())
+	if (ClothConfigs.Num())
 	{
-		// No self collision, can't generate data
-		return;
+		for(UClothLODDataCommon* Lod : ClothLodData)
+		{
+			check(Lod);
+			Lod->ClothPhysicalMeshData.BuildSelfCollisionData(ClothConfigs);
+		}
 	}
-	for(UClothLODDataCommon* Lod : ClothLodData)
+}
+
+void UClothingAssetCommon::PostLoad()
+{
+	Super::PostLoad();
+
+	const int32 AnimPhysCustomVersion = GetLinkerCustomVersion(FAnimPhysObjectVersion::GUID);
+	const int32 ClothingCustomVersion = GetLinkerCustomVersion(FClothingAssetCustomVersion::GUID);
+
+	if (ClothingCustomVersion < FClothingAssetCustomVersion::MovePropertiesToCommonBaseClasses)
+	{
+		// Remap legacy struct FClothConfig to new config objects
+		for (TPair<FName, UClothConfigBase*>& ClothConfig : ClothConfigs)
+		{
+			if (UClothConfigCommon* const ClothConfigCommon = Cast<UClothConfigCommon>(ClothConfig.Value))
+			{
+				ClothConfigCommon->MigrateFrom(ClothConfig_DEPRECATED);
+			}
+		}
+
+		// Remap legacy struct FClothLODData to class UClothLODDataCommon
+		for (const FClothLODData_Legacy& ClothLODData_Legacy : LodData_DEPRECATED)
+		{
+			const int32 Idx = AddNewLod();
+			ClothLODData_Legacy.MigrateTo(ClothLodData[Idx]);
+		}
+		LodData_DEPRECATED.Empty();
+	}
+	if(AnimPhysCustomVersion < FAnimPhysObjectVersion::AddedClothingMaskWorkflow)
+	{
+#if WITH_EDITORONLY_DATA
+		// Convert current parameters to masks
+		for(UClothLODDataCommon* LodPtr : ClothLodData)
+		{
+			check(LodPtr);
+			UClothLODDataCommon& Lod = *LodPtr;
+			const FClothPhysicalMeshData& PhysMesh = Lod.ClothPhysicalMeshData;
+
+			// Didn't do anything previously - clear out in case there's something in there
+			// so we can use it correctly now.
+			Lod.ParameterMasks.Reset(3);
+
+			// Max distances (Always present)
+			Lod.ParameterMasks.AddDefaulted();
+			FPointWeightMap& MaxDistanceMask = Lod.ParameterMasks.Last();
+			const FPointWeightMap& MaxDistances = PhysMesh.GetWeightMap(EWeightMapTargetCommon::MaxDistance);
+			MaxDistanceMask.Initialize(MaxDistances, EWeightMapTargetCommon::MaxDistance);
+
+			// Following params are only added if necessary, if we don't have any backstop
+			// radii then there's no backstops.
+			const FPointWeightMap* const BackstopRadiuses = PhysMesh.FindWeightMap(EWeightMapTargetCommon::BackstopRadius);
+			if(BackstopRadiuses && !BackstopRadiuses->IsZeroed())
+			{
+				// Backstop radii
+				Lod.ParameterMasks.AddDefaulted();
+				FPointWeightMap& BackstopRadiusMask = Lod.ParameterMasks.Last();
+				BackstopRadiusMask.Initialize(*BackstopRadiuses, EWeightMapTargetCommon::BackstopRadius);
+
+				// Backstop distances
+				Lod.ParameterMasks.AddDefaulted();
+				FPointWeightMap& BackstopDistanceMask = Lod.ParameterMasks.Last();
+				const FPointWeightMap& BackstopDistances = PhysMesh.GetWeightMap(EWeightMapTargetCommon::BackstopDistance);
+				BackstopDistanceMask.Initialize(BackstopDistances, EWeightMapTargetCommon::BackstopDistance);
+			}
+			
+		}
+#endif
+
+		// Make sure we're transactional
+		SetFlags(RF_Transactional);
+	}
+
+#if WITH_EDITORONLY_DATA
+	// Fix content imported before we kept vertex colors
+	if(ClothingCustomVersion < FClothingAssetCustomVersion::AddVertexColorsToPhysicalMesh)
+	{
+		for (UClothLODDataCommon* Lod : ClothLodData)
+		{
+			const int32 NumVerts = Lod->ClothPhysicalMeshData.Vertices.Num(); // number of verts
+
+			Lod->ClothPhysicalMeshData.VertexColors.Reset();
+			Lod->ClothPhysicalMeshData.VertexColors.AddUninitialized(NumVerts);
+			for (int32 VertIdx = 0; VertIdx < NumVerts; VertIdx++)
+			{
+				Lod->ClothPhysicalMeshData.VertexColors[VertIdx] = FColor::White;
+			}
+		}
+	}
+#endif // WITH_EDITORONLY_DATA
+
+#if WITH_EDITOR
+	if(AnimPhysCustomVersion < FAnimPhysObjectVersion::CacheClothMeshInfluences)
+	{
+		// Rebuild data cache
+		InvalidateCachedData();
+	}
+#endif
+
+	// After fixing the content, we are ready to call functions that rely on it
+	BuildSelfCollisionData();
+#if WITH_EDITORONLY_DATA
+	CalculateReferenceBoneIndex();
+#endif
+
+	// Migrate simulation dependent config parameters to their new abstract equivalent
+	if (ClothSimConfig_DEPRECATED)
 	{
-		check(Lod && Lod->PhysicalMeshData);
-		Lod->PhysicalMeshData->BuildSelfCollisionData(ClothSimConfig);
+		SetClothConfig(ClothSimConfig_DEPRECATED);
+		ClothSimConfig_DEPRECATED = nullptr;
+	}
+	if (ChaosClothSimConfig_DEPRECATED)
+	{
+		SetClothConfig(ChaosClothSimConfig_DEPRECATED);
+		ChaosClothSimConfig_DEPRECATED = nullptr;
+	}
+	if (ClothSharedSimConfig_DEPRECATED)
+	{
+		SetClothConfig(ClothSharedSimConfig_DEPRECATED);
+		ClothSharedSimConfig_DEPRECATED = nullptr;
 	}
 }
 
+void UClothingAssetCommon::Serialize(FArchive& Ar)
+{
+	Super::Serialize(Ar);
+	Ar.UsingCustomVersion(FAnimPhysObjectVersion::GUID);
+	Ar.UsingCustomVersion(FClothingAssetCustomVersion::GUID);
+}
+
+#if WITH_EDITOR
+
+void UClothingAssetCommon::InvalidateCachedData()
+{
+	for(UClothLODDataCommon* CurrentLodData : ClothLodData)
+	{
+		check(CurrentLodData);
+		// Recalculate inverse masses for the physical mesh particles
+		FClothPhysicalMeshData& PhysMesh = CurrentLodData->ClothPhysicalMeshData;
+		check(PhysMesh.Indices.Num() % 3 == 0);
+
+		TArray<float>& InvMasses = PhysMesh.InverseMasses;
+
+		const int32 NumVerts = PhysMesh.Vertices.Num();
+		InvMasses.Empty(NumVerts);
+		InvMasses.AddZeroed(NumVerts);
+
+		for(int32 TriBaseIndex = 0; TriBaseIndex < PhysMesh.Indices.Num(); TriBaseIndex += 3)
+		{
+			const int32 Index0 = PhysMesh.Indices[TriBaseIndex];
+			const int32 Index1 = PhysMesh.Indices[TriBaseIndex + 1];
+			const int32 Index2 = PhysMesh.Indices[TriBaseIndex + 2];
+
+			const FVector AB = PhysMesh.Vertices[Index1] - PhysMesh.Vertices[Index0];
+			const FVector AC = PhysMesh.Vertices[Index2] - PhysMesh.Vertices[Index0];
+			const float TriArea = FVector::CrossProduct(AB, AC).Size();
+
+			InvMasses[Index0] += TriArea;
+			InvMasses[Index1] += TriArea;
+			InvMasses[Index2] += TriArea;
+		}
+
+		PhysMesh.NumFixedVerts = 0;
+
+		const FPointWeightMap* const MaxDistances = PhysMesh.FindWeightMap(EWeightMapTargetCommon::MaxDistance);
+		if (MaxDistances && MaxDistances->Num() > 0)
+		{
+			float MassSum = 0.0f;
+			for (int32 CurrVertIndex = 0; CurrVertIndex < NumVerts; ++CurrVertIndex)
+			{
+				float& InvMass = InvMasses[CurrVertIndex];
+				const float& MaxDistance = (*MaxDistances)[CurrVertIndex];
+
+				if (MaxDistance < SMALL_NUMBER)   // For consistency, the default Threshold should be 0.1, not SMALL_NUMBER. But for backward compatibility it needs to be SMALL_NUMBER for now.
+				{
+					InvMass = 0.0f;
+					++PhysMesh.NumFixedVerts;
+				}
+				else
+				{
+					MassSum += InvMass;
+				}
+			}
+
+			if (MassSum > 0.0f)
+			{
+				const float MassScale = (float)(NumVerts - PhysMesh.NumFixedVerts) / MassSum;
+				for (float& InvMass : InvMasses)
+				{
+					if (InvMass != 0.0f)
+					{
+						InvMass *= MassScale;
+						InvMass = 1.0f / InvMass;
+					}
+				}
+			}
+		}
+		else
+		{
+			// Otherwise, go fully kinematic.
+			for(int32 CurrVertIndex = 0; CurrVertIndex < NumVerts; ++CurrVertIndex)
+			{
+				InvMasses[CurrVertIndex] = 0.0f;
+			}
+			PhysMesh.NumFixedVerts = NumVerts;
+		}
+
+		// Calculate number of influences per vertex
+		for(int32 VertIndex = 0; VertIndex < NumVerts; ++VertIndex)
+		{
+			FClothVertBoneData& BoneData = PhysMesh.BoneData[VertIndex];
+			const uint16* BoneIndices = BoneData.BoneIndices;
+			const float* BoneWeights = BoneData.BoneWeights;
+
+			BoneData.NumInfluences = MAX_TOTAL_INFLUENCES;
+
+			int32 NumInfluences = 0;
+			for(int32 InfluenceIndex = 0; InfluenceIndex < MAX_TOTAL_INFLUENCES; ++InfluenceIndex)
+			{
+				if(BoneWeights[InfluenceIndex] == 0.0f || BoneIndices[InfluenceIndex] == INDEX_NONE)
+				{
+					BoneData.NumInfluences = NumInfluences;
+					break;
+				}
+				++NumInfluences;
+			}
+		}
+	}
+}
+
+int32 UClothingAssetCommon::AddNewLod()
+{
+	const int32 Idx = ClothLodData.AddDefaulted();
+	ClothLodData[Idx] = NewObject<UClothLODDataCommon>(this, UClothLODDataCommon::StaticClass(), NAME_None, RF_Transactional);
+	return Idx;
+}
+
+void UClothingAssetCommon::PostEditChangeChainProperty(FPropertyChangedChainEvent& ChainEvent)
+{
+	Super::PostEditChangeChainProperty(ChainEvent);
+
+	bool bReregisterComponents = false;
+
+	if (ChainEvent.ChangeType != EPropertyChangeType::Interactive)
+	{
+		if (ChainEvent.Property->GetFName() == FName("SelfCollisionRadius") ||
+			ChainEvent.Property->GetFName() == FName("SelfCollisionCullScale"))
+		{
+			BuildSelfCollisionData();
+			bReregisterComponents = true;
+		}
+		else if(ChainEvent.Property->GetFName() == GET_MEMBER_NAME_CHECKED(UClothingAssetCommon, PhysicsAsset))
+		{
+			bReregisterComponents = true;
+		}
+		else
+		{
+			// Other properties just require a config refresh
+			ForEachInteractorUsingClothing([](UClothingSimulationInteractor* InInteractor)
+			{
+				if (InInteractor)
+				{
+					InInteractor->ClothConfigUpdated();
+				}
+			});
+		}
+	}
+
+	if (bReregisterComponents)
+	{
+		ReregisterComponentsUsingClothing();
+	}
+}
+
+#endif // WITH_EDITOR
+
 
 #undef LOCTEXT_NAMESPACE
diff --git a/Engine/Source/Runtime/ClothingSystemRuntimeCommon/Private/ClothingAssetCustomVersion.cpp b/Engine/Source/Runtime/ClothingSystemRuntimeCommon/Private/ClothingAssetCustomVersion.cpp
new file mode 100644
--- /dev/null
+++ b/Engine/Source/Runtime/ClothingSystemRuntimeCommon/Private/ClothingAssetCustomVersion.cpp
@@ -0,0 +1,9 @@
+// Copyright 1998-2019 Epic Games, Inc. All Rights Reserved.
+
+#include "ClothingAssetCustomVersion.h"
+#include "Serialization/CustomVersion.h"
+
+const FGuid FClothingAssetCustomVersion::GUID(0xFB680AF2, 0x59EF4BA3, 0xBAA819B5, 0x73C8443D);
+
+// Register the custom version with core
+FCustomVersionRegistration GRegisterClothingAssetCustomVersion(FClothingAssetCustomVersion::GUID, FClothingAssetCustomVersion::LatestVersion, TEXT("ClothingAssetVer"));
diff --git a/Engine/Source/Runtime/ClothingSystemRuntimeCommon/Private/ClothingAssetCustomVersion.h b/Engine/Source/Runtime/ClothingSystemRuntimeCommon/Private/ClothingAssetCustomVersion.h
new file mode 100644
--- /dev/null
+++ b/Engine/Source/Runtime/ClothingSystemRuntimeCommon/Private/ClothingAssetCustomVersion.h
@@ -0,0 +1,30 @@
+// Copyright 1998-2019 Epic Games, Inc. All Rights Reserved.
+
+#pragma once
+
+#include "CoreMinimal.h"
+#include "Misc/Guid.h"
+
+// Custom serialization version for clothing assets
+struct FClothingAssetCustomVersion
+{
+	enum Type
+	{
+		// Before any version changes were made
+		BeforeCustomVersionWasAdded = 0,
+		// Added storage of vertex colors with sim data, for editor usage
+		AddVertexColorsToPhysicalMesh = 1,
+		// Changed memory layouts by moving some properties to base classes
+		MovePropertiesToCommonBaseClasses = 2,
+
+		// -----<new versions can be added above this line>-------------------------------------------------
+		VersionPlusOne,
+		LatestVersion = VersionPlusOne - 1
+	};
+
+	// The GUID for this custom version number
+	const static FGuid GUID;
+
+private:
+	FClothingAssetCustomVersion() {}
+};
diff --git a/Engine/Source/Runtime/ClothingSystemRuntimeCommon/Private/Utils/ClothingMeshUtils.cpp b/Engine/Source/Runtime/ClothingSystemRuntimeCommon/Private/Utils/ClothingMeshUtils.cpp
--- a/Engine/Source/Runtime/ClothingSystemRuntimeCommon/Private/Utils/ClothingMeshUtils.cpp
+++ b/Engine/Source/Runtime/ClothingSystemRuntimeCommon/Private/Utils/ClothingMeshUtils.cpp
@@ -2,7 +2,7 @@
 
 #include "Utils/ClothingMeshUtils.h"
 
-#include "ClothPhysicalMeshDataBase.h"
+#include "ClothPhysicalMeshData.h"
 
 #include "Math/UnrealMathUtility.h"
 #include "Logging/LogMacros.h"
@@ -21,16 +21,16 @@ namespace ClothingMeshUtils
 {
 	// Explicit template instantiations of SkinPhysicsMesh
 	template
-	void CLOTHINGSYSTEMRUNTIMECOMMON_API SkinPhysicsMesh<true, false>(const TArray<int32>& BoneMap, const UClothPhysicalMeshDataBase& InMesh, const FTransform& RootBoneTransform,
+	void CLOTHINGSYSTEMRUNTIMECOMMON_API SkinPhysicsMesh<true, false>(const TArray<int32>& BoneMap, const FClothPhysicalMeshData& InMesh, const FTransform& RootBoneTransform,
 		const FMatrix* InBoneMatrices, const int32 InNumBoneMatrices, TArray<FVector>& OutPositions, TArray<FVector>& OutNormals, uint32 ArrayOffset);
 	template
-	void CLOTHINGSYSTEMRUNTIMECOMMON_API SkinPhysicsMesh<false, true>(const TArray<int32>& BoneMap, const UClothPhysicalMeshDataBase& InMesh, const FTransform& RootBoneTransform,
+	void CLOTHINGSYSTEMRUNTIMECOMMON_API SkinPhysicsMesh<false, true>(const TArray<int32>& BoneMap, const FClothPhysicalMeshData& InMesh, const FTransform& RootBoneTransform,
 		const FMatrix* InBoneMatrices, const int32 InNumBoneMatrices, TArray<FVector>& OutPositions, TArray<FVector>& OutNormals, uint32 ArrayOffset);
 
 	template<bool bInPlaceOutput, bool bRemoveScaleAndInvertPostTransform>
 	void SkinPhysicsMesh(
 		const TArray<int32>& BoneMap,
-		const UClothPhysicalMeshDataBase& InMesh,
+		const FClothPhysicalMeshData& InMesh,
 		const FTransform& PostTransform,
 		const FMatrix* InBoneMatrices,
 		const int32 InNumBoneMatrices,
diff --git a/Engine/Source/Runtime/ClothingSystemRuntimeCommon/Public/ClothConfig.h b/Engine/Source/Runtime/ClothingSystemRuntimeCommon/Public/ClothConfig.h
new file mode 100644
--- /dev/null
+++ b/Engine/Source/Runtime/ClothingSystemRuntimeCommon/Public/ClothConfig.h
@@ -0,0 +1,20 @@
+// Copyright 1998-2019 Epic Games, Inc. All Rights Reserved.
+#pragma once
+
+#include "ClothConfigBase.h"
+#include "ClothConfig.generated.h"
+
+struct FClothConfig_Legacy;
+
+/** Common configuration base class. */
+UCLASS(Abstract)
+class CLOTHINGSYSTEMRUNTIMECOMMON_API UClothConfigCommon : public UClothConfigBase
+{
+	GENERATED_BODY()
+public:
+	UClothConfigCommon();
+	virtual ~UClothConfigCommon();
+
+	/** Migrate from the legacy FClothConfig structure. */
+	virtual void MigrateFrom(const FClothConfig_Legacy&) {}
+};
diff --git a/Engine/Source/Runtime/ClothingSystemRuntimeCommon/Public/ClothConfig_Legacy.h b/Engine/Source/Runtime/ClothingSystemRuntimeCommon/Public/ClothConfig_Legacy.h
new file mode 100644
--- /dev/null
+++ b/Engine/Source/Runtime/ClothingSystemRuntimeCommon/Public/ClothConfig_Legacy.h
@@ -0,0 +1,202 @@
+// Copyright 1998-2019 Epic Games, Inc. All Rights Reserved.
+#pragma once
+
+#include "UObject/ObjectMacros.h"
+
+#include "ClothConfig_Legacy.generated.h"
+
+/**
+ * Deprecated, legacy definition kept for backward compatibility only.
+ * Use EClothingWindMethodNv instead.
+ * Redirected from the now defunct ClothingSystemRuntime module.
+ */
+UENUM()
+enum class EClothingWindMethod_Legacy : uint8
+{
+	// Use legacy wind mode, where accelerations are modified directly by the simulation
+	// with no regard for drag or lift
+	Legacy,
+
+	// Use updated wind calculation for NvCloth based solved taking into account
+	// drag and lift, this will require those properties to be correctly set in
+	// the clothing configuration
+	Accurate
+};
+
+/**
+ * Deprecated, legacy definition kept for backward compatibility only.
+ * Use FClothConstraintSetupNv instead.
+ * Redirected from the now defunct ClothingSystemRuntime module.
+ */
+USTRUCT()
+struct FClothConstraintSetup_Legacy
+{
+	GENERATED_BODY()
+
+	FClothConstraintSetup_Legacy()
+		: Stiffness(1.0f)
+		, StiffnessMultiplier(1.0f)
+		, StretchLimit(1.0f)
+		, CompressionLimit(1.0f)
+	{}
+
+	// How stiff this constraint is, this affects how closely it will follow the desired position
+	UPROPERTY()
+	float Stiffness;
+
+	// A multiplier affecting the above value
+	UPROPERTY()
+	float StiffnessMultiplier;
+
+	// The hard limit on how far this constraint can stretch
+	UPROPERTY()
+	float StretchLimit;
+
+	// The hard limit on how far this constraint can compress
+	UPROPERTY()
+	float CompressionLimit;
+};
+
+/**
+ * Deprecated, legacy definition kept for backward compatibility only.
+ * Inherit new config class from UClothConfigCommon instead.
+ * Redirected from the now defunct ClothingSystemRuntime module.
+ */
+USTRUCT()
+struct CLOTHINGSYSTEMRUNTIMECOMMON_API FClothConfig_Legacy
+{
+	GENERATED_BODY()
+
+	FClothConfig_Legacy()
+		: WindMethod(EClothingWindMethod_Legacy::Legacy)
+		, SelfCollisionRadius(0.0f)
+		, SelfCollisionStiffness(0.0f)
+		, SelfCollisionCullScale(1.0f)
+		, Damping(0.4f)
+		, Friction(0.1f)
+		, WindDragCoefficient(0.02f/100.0f)
+		, WindLiftCoefficient(0.02f/100.0f)
+		, LinearDrag(0.2f)
+		, AngularDrag(0.2f)
+		, LinearInertiaScale(1.0f)
+		, AngularInertiaScale(1.0f)
+		, CentrifugalInertiaScale(1.0f)
+		, SolverFrequency(120.0f)
+		, StiffnessFrequency(100.0f)
+		, GravityScale(1.0f)
+		, GravityOverride(FVector::ZeroVector)
+		, bUseGravityOverride(false)
+		, TetherStiffness(1.0f)
+		, TetherLimit(1.0f)
+		, CollisionThickness(1.0f)
+		, AnimDriveSpringStiffness(1.0f)
+		, AnimDriveDamperStiffness(1.0f)
+	{}
+
+	// How wind should be processed, Accurate uses drag and lift to make the cloth react differently, legacy applies similar forces to all clothing without drag and lift (similar to APEX)
+	UPROPERTY()
+	EClothingWindMethod_Legacy WindMethod;
+
+	// Constraint data for vertical constraints
+	UPROPERTY()
+	FClothConstraintSetup_Legacy VerticalConstraintConfig;
+
+	// Constraint data for horizontal constraints
+	UPROPERTY()
+	FClothConstraintSetup_Legacy HorizontalConstraintConfig;
+
+	// Constraint data for bend constraints
+	UPROPERTY()
+	FClothConstraintSetup_Legacy BendConstraintConfig;
+
+	// Constraint data for shear constraints
+	UPROPERTY()
+	FClothConstraintSetup_Legacy ShearConstraintConfig;
+
+	// Size of self collision spheres centered on each vert
+	UPROPERTY()
+	float SelfCollisionRadius;
+
+	// Stiffness of the spring force that will resolve self collisions
+	UPROPERTY()
+	float SelfCollisionStiffness;
+
+	UPROPERTY()
+	float SelfCollisionCullScale;
+
+	// Damping of particle motion per-axis
+	UPROPERTY()
+	FVector Damping;
+
+	// Friction of the surface when colliding
+	UPROPERTY()
+	float Friction;
+
+	// Drag coefficient for wind calculations, higher values mean wind has more lateral effect on cloth
+	UPROPERTY()
+	float WindDragCoefficient;
+
+	// Lift coefficient for wind calculations, higher values make cloth rise easier in wind
+	UPROPERTY()
+	float WindLiftCoefficient;
+
+	// Drag applied to linear particle movement per-axis
+	UPROPERTY()
+	FVector LinearDrag;
+
+	// Drag applied to angular particle movement, higher values should limit material bending (per-axis)
+	UPROPERTY()
+	FVector AngularDrag;
+
+	// Scale for linear particle inertia, how much movement should translate to linear motion (per-axis)
+	UPROPERTY()
+	FVector LinearInertiaScale;
+
+	// Scale for angular particle inertia, how much movement should translate to angular motion (per-axis)
+	UPROPERTY()
+	FVector AngularInertiaScale;
+
+	// Scale for centrifugal particle inertia, how much movement should translate to angular motion (per-axis)
+	UPROPERTY()
+	FVector CentrifugalInertiaScale;
+
+	// Frequency of the position solver, lower values will lead to stretchier, bouncier cloth
+	UPROPERTY()
+	float SolverFrequency;
+
+	// Frequency for stiffness calculations, lower values will degrade stiffness of constraints
+	UPROPERTY()
+	float StiffnessFrequency;
+
+	// Scale of gravity effect on particles
+	UPROPERTY()
+	float GravityScale;
+
+	// Direct gravity override value
+	UPROPERTY()
+	FVector GravityOverride;
+
+	// Use gravity override value vs gravity scale 
+	UPROPERTY()
+	bool bUseGravityOverride;
+
+	// Scale for stiffness of particle tethers between each other
+	UPROPERTY()
+	float TetherStiffness;
+
+	// Scale for the limit of particle tethers (how far they can separate)
+	UPROPERTY()
+	float TetherLimit;
+
+	// 'Thickness' of the simulated cloth, used to adjust collisions
+	UPROPERTY()
+	float CollisionThickness;
+
+	// Default spring stiffness for anim drive if an anim drive is in use
+	UPROPERTY()
+	float AnimDriveSpringStiffness;
+
+	// Default damper stiffness for anim drive if an anim drive is in use
+	UPROPERTY()
+	float AnimDriveDamperStiffness;
+};
diff --git a/Engine/Source/Runtime/ClothingSystemRuntimeCommon/Public/ClothLODData.h b/Engine/Source/Runtime/ClothingSystemRuntimeCommon/Public/ClothLODData.h
--- a/Engine/Source/Runtime/ClothingSystemRuntimeCommon/Public/ClothLODData.h
+++ b/Engine/Source/Runtime/ClothingSystemRuntimeCommon/Public/ClothLODData.h
@@ -1,16 +1,15 @@
 // Copyright 1998-2019 Epic Games, Inc. All Rights Reserved.
 #pragma once
 
-#include "ClothPhysicalMeshDataBase.h"
-
+#include "ClothPhysicalMeshData.h"
+#include "ClothPhysicalMeshDataBase_Legacy.h"
 #include "ClothCollisionData.h"
 #include "PointWeightMap.h"
-
 #include "SkeletalMeshTypes.h"
 
 #include "ClothLODData.generated.h"
 
-/** */
+/** Common Cloth LOD representation for all clothing assets. */
 UCLASS()
 class CLOTHINGSYSTEMRUNTIMECOMMON_API UClothLODDataCommon : public UObject
 {
@@ -19,9 +18,13 @@ public:
 	UClothLODDataCommon(const FObjectInitializer& Init);
 	virtual ~UClothLODDataCommon();
 
+	// Deprecated, use ClothPhysicalMeshData instead
+	UPROPERTY()
+	UClothPhysicalMeshDataBase_Legacy* PhysicalMeshData_DEPRECATED;
+
 	// Raw phys mesh data
 	UPROPERTY(EditAnywhere, Category = SimMesh)
-	UClothPhysicalMeshDataBase* PhysicalMeshData;
+	FClothPhysicalMeshData ClothPhysicalMeshData;
 
 	// Collision primitive and covex data for clothing collisions
 	UPROPERTY(EditAnywhere, Category = Collision)
@@ -33,11 +36,10 @@ public:
 	TArray<FPointWeightMap> ParameterMasks;
 
 	// Get all available parameter masks for the specified target
-	//void GetParameterMasksForTarget(const MaskTarget_PhysMesh& InTarget, TArray<FPointWeightMap*>& OutMasks);
 	void GetParameterMasksForTarget(const uint8 InTarget, TArray<FPointWeightMap*>& OutMasks);
 #endif // WITH_EDITORONLY_DATA
 #if WITH_EDITOR
-	/** Copy \c ParameterMasks to corresponding targets in \c PhysicalMeshData. */
+	/** Copy \c ParameterMasks to corresponding targets in \c ClothPhysicalMeshData. */
 	void PushWeightsToMesh();
 #endif
 	// Skinning data for transitioning from a higher detail LOD to this one
@@ -46,5 +48,9 @@ public:
 	// Skinning data for transitioning from a lower detail LOD to this one
 	TArray<FMeshToMeshVertData> TransitionDownSkinData;
 
+	// Custom serialize for transition
 	virtual void Serialize(FArchive& Ar) override;
+
+	// Migrate deprecated properties
+	virtual void PostLoad() override;
 };
diff --git a/Engine/Source/Runtime/ClothingSystemRuntimeCommon/Public/ClothLODData_Legacy.h b/Engine/Source/Runtime/ClothingSystemRuntimeCommon/Public/ClothLODData_Legacy.h
new file mode 100644
--- /dev/null
+++ b/Engine/Source/Runtime/ClothingSystemRuntimeCommon/Public/ClothLODData_Legacy.h
@@ -0,0 +1,96 @@
+// Copyright 1998-2019 Epic Games, Inc. All Rights Reserved.
+#pragma once
+
+#include "ClothCollisionData.h"
+#include "SkeletalMeshTypes.h"
+#include "ClothVertBoneData.h"
+#include "ClothPhysicalMeshData.h"
+
+#include "ClothLODData_Legacy.generated.h"
+
+class UClothLODDataCommon;
+struct FPointWeightMap;
+
+/**
+ * Deprecated, legacy definition kept for backward compatibility only.
+ * Use FPointWeightMap instead.
+ * Redirected from the now defunct ClothingSystemRuntime module.
+ */
+USTRUCT()
+struct CLOTHINGSYSTEMRUNTIMECOMMON_API FClothParameterMask_Legacy
+{
+	GENERATED_BODY();
+
+	FClothParameterMask_Legacy();
+
+	void MigrateTo(FPointWeightMap* Weights) const;
+
+	/** Name of the mask, mainly for users to differentiate */
+	UPROPERTY()
+	FName MaskName;
+
+	/** The currently targeted parameter for the mask */
+	UPROPERTY()
+	EWeightMapTargetCommon CurrentTarget;
+
+	/** The maximum value currently in the mask value array */
+	UPROPERTY()
+	float MaxValue_DEPRECATED;
+
+	/** The maximum value currently in the mask value array */
+	UPROPERTY()
+	float MinValue_DEPRECATED;
+
+	/** The actual values stored in the mask */
+	UPROPERTY()
+	TArray<float> Values;
+
+	/** Whether this mask is enabled and able to effect final mesh values */
+	UPROPERTY()
+	bool bEnabled;
+};
+
+/**
+ * Deprecated, legacy definition kept for backward compatibility only.
+ * Use UClothLODDataCommon instead.
+ * Redirected from the now defunct ClothingSystemRuntime module.
+ */
+USTRUCT()
+struct CLOTHINGSYSTEMRUNTIMECOMMON_API FClothLODData_Legacy
+{
+	GENERATED_BODY()
+
+	// Raw phys mesh data
+	UPROPERTY()
+	FClothPhysicalMeshData PhysicalMeshData;
+
+	// Collision primitive and covex data for clothing collisions
+	UPROPERTY()
+	FClothCollisionData CollisionData;
+
+#if WITH_EDITORONLY_DATA
+	// Parameter masks defining the physics mesh masked data
+	UPROPERTY()
+	TArray<FClothParameterMask_Legacy> ParameterMasks;
+#endif // WITH_EDITORONLY_DATA
+
+	// Skinning data for transitioning from a higher detail LOD to this one
+	TArray<FMeshToMeshVertData> TransitionUpSkinData;
+
+	// Skinning data for transitioning from a lower detail LOD to this one
+	TArray<FMeshToMeshVertData> TransitionDownSkinData;
+
+	bool Serialize(FArchive& Ar);
+
+	// Migrate this legacy Nv struct to the new class format (called by UClothingAssetCommon::Postload())
+	void MigrateTo(UClothLODDataCommon* LodData) const;
+};
+
+template<>
+struct TStructOpsTypeTraits<FClothLODData_Legacy> : public TStructOpsTypeTraitsBase2<FClothLODData_Legacy>
+{
+	enum
+	{
+		WithSerializer = true,
+	};
+};
diff --git a/Engine/Source/Runtime/ClothingSystemRuntimeCommon/Public/ClothPhysicalMeshData.h b/Engine/Source/Runtime/ClothingSystemRuntimeCommon/Public/ClothPhysicalMeshData.h
new file mode 100644
--- /dev/null
+++ b/Engine/Source/Runtime/ClothingSystemRuntimeCommon/Public/ClothPhysicalMeshData.h
@@ -0,0 +1,124 @@
+// Copyright 1998-2019 Epic Games, Inc. All Rights Reserved.
+#pragma once
+
+#include "ClothVertBoneData.h"
+#include "PointWeightMap.h"
+#include "Containers/Array.h"
+#include "Containers/Map.h"
+
+#include "ClothPhysicalMeshData.generated.h"
+
+class UClothPhysicalMeshDataBase_Legacy;
+class UClothConfigBase;
+
+/** Spatial simulation data for a mesh. */
+USTRUCT()
+struct CLOTHINGSYSTEMRUNTIMECOMMON_API FClothPhysicalMeshData
+{
+	GENERATED_BODY()
+
+	/** Construct an empty cloth physical mesh with default common targets. */
+	FClothPhysicalMeshData();
+
+	/** Migrate away the deprecated properties. */
+	void MigrateTo(FClothPhysicalMeshData& ClothPhysicalMeshData) const;
+
+	/** Migrate from the nv legacy physical mesh data class. */
+	void MigrateFrom(const UClothPhysicalMeshDataBase_Legacy* ClothPhysicalMeshDataBase);
+
+	/** Reset the default common targets for this cloth physical mesh. */
+	void Reset(const int32 InNumVerts, const int32 InNumIndices);
+
+	/** Clear out any default weight maps and delete any other ones. */
+	void ClearWeightMaps();
+
+	/** Build the self collision indices for the relevant config. */
+	void BuildSelfCollisionData(const TMap<FName, UClothConfigBase*>& ClothConfigs);
+
+	/** Retrieve whether a vertex weight array has already been registered. */
+	template<typename T>
+	bool HasWeightMap(const T Target) const
+	{ return WeightMaps.Contains((uint32)Target); }
+
+	/** Retrieve a pointer to a registered vertex weight array by unique @param Id, or nullptr if none is found. */
+	template<typename T>
+	const FPointWeightMap* FindWeightMap(const T Target) const
+	{ return WeightMaps.Find((uint32)Target); }
+
+	/** Retrieve a pointer to a registered vertex weight array by unique @param Id, or nullptr if none is found. */
+	template<typename T>
+	FPointWeightMap* FindWeightMap(const T Target)
+	{ return WeightMaps.Find((uint32)Target); }
+
+	/** Retrieve a pointer to a registered vertex weight array by unique @param Id, or add one if it doesn't exist already. */
+	template<typename T>
+	FPointWeightMap& AddWeightMap(const T Target)
+	{ return WeightMaps.Add((uint32)Target); }
+
+	/** Retrieve a pointer to a registered vertex weight array by unique @param Id, or add one if it doesn't exist already. */
+	template<typename T>
+	FPointWeightMap& FindOrAddWeightMap(const T Target)
+	{ return WeightMaps.FindOrAdd((uint32)Target); }
+
+	/** Retrieve a registered vertex weight array by unique @param Id. The array must exists or this function will assert. */
+	template<typename T>
+	const FPointWeightMap& GetWeightMap(const T Target) const
+	{ return WeightMaps[(uint32)Target]; }
+
+	/** Retrieve a registered vertex weight array by unique @param Id. The array must exists or this function will assert. */
+	template<typename T>
+	FPointWeightMap& GetWeightMap(const T Target)
+	{ return WeightMaps[(uint32)Target]; }
+
+	// Positions of each simulation vertex
+	UPROPERTY(EditAnywhere, Category = SimMesh)
+	TArray<FVector> Vertices;
+
+	// Normal at each vertex
+	UPROPERTY(EditAnywhere, Category = SimMesh)
+	TArray<FVector> Normals;
+
+#if WITH_EDITORONLY_DATA
+	// Color at each vertex
+	UPROPERTY(EditAnywhere, Category = SimMesh)
+	TArray<FColor> VertexColors;
+#endif // WITH_EDITORONLY_DATA
+
+	// Indices of the simulation mesh triangles
+	UPROPERTY(EditAnywhere, Category = SimMesh)
+	TArray<uint32> Indices;
+
+	// The weight maps, or masks, used by this mesh, sorted by their target id
+	UPROPERTY(EditAnywhere, Category = SimMesh)
+	TMap<uint32, FPointWeightMap> WeightMaps;
+
+	// Inverse mass for each vertex in the physical mesh
+	UPROPERTY(EditAnywhere, Category = SimMesh)
+	TArray<float> InverseMasses;
+
+	// Indices and weights for each vertex, used to skin the mesh to create the reference pose
+	UPROPERTY(EditAnywhere, Category = SimMesh)
+	TArray<FClothVertBoneData> BoneData;
+
+	// Maximum number of bone weights of any vetex
+	UPROPERTY(EditAnywhere, Category = SimMesh)
+	int32 MaxBoneWeights;
+
+	// Number of fixed verts in the simulation mesh (fixed verts are just skinned and do not simulate)
+	UPROPERTY(EditAnywhere, Category = SimMesh)
+	int32 NumFixedVerts;
+
+	// Valid indices to use for self collisions (reduced set of Indices)
+	UPROPERTY(EditAnywhere, Category = SimMesh)
+	TArray<uint32> SelfCollisionIndices;
+
+	// Deprecated. Use WeightMaps instead.
+	UPROPERTY()
+	TArray<float> MaxDistances_DEPRECATED;
+	UPROPERTY()
+	TArray<float> BackstopDistances_DEPRECATED;
+	UPROPERTY()
+	TArray<float> BackstopRadiuses_DEPRECATED;
+	UPROPERTY()
+	TArray<float> AnimDriveMultipliers_DEPRECATED;
+};
diff --git a/Engine/Source/Runtime/ClothingSystemRuntimeCommon/Public/ClothingAsset.h b/Engine/Source/Runtime/ClothingSystemRuntimeCommon/Public/ClothingAsset.h
--- a/Engine/Source/Runtime/ClothingSystemRuntimeCommon/Public/ClothingAsset.h
+++ b/Engine/Source/Runtime/ClothingSystemRuntimeCommon/Public/ClothingAsset.h
@@ -4,6 +4,8 @@
 #include "ClothingAssetBase.h"
 #include "ClothConfigBase.h"
 #include "ClothLODData.h"
+#include "ClothConfig_Legacy.h"
+#include "ClothLODData_Legacy.h"
 #include "ClothingSimulationInteractor.h"
 
 #include "ClothingAsset.generated.h"
@@ -107,9 +109,6 @@ public:
 	virtual void UnbindFromSkeletalMesh(USkeletalMesh* InSkelMesh) override;
 	virtual void UnbindFromSkeletalMesh(USkeletalMesh* InSkelMesh, const int32 InMeshLodIndex) override;
 
-	void ReregisterComponentsUsingClothing();
-	void ForEachInteractorUsingClothing(TFunction<void (UClothingSimulationInteractor*)> Func);
-
 	/** 
 	 * Callback envoked after weights have been edited.
 	 * Calls \c PushWeightsToMesh() on each \c ClothLodData, and invalidates cached data. 
@@ -131,50 +130,93 @@ public:
 	 */
 	virtual void RefreshBoneMapping(USkeletalMesh* InSkelMesh) override;
 
-	/**
-	 * Calculates the prefered root bone for the simulation.
-	 */
+	/** Calculates the preferred root bone for the simulation. */
 	void CalculateReferenceBoneIndex();
 
-	/**
-	 * Returns \c true if \p InLodIndex is a valid LOD id (index into \c ClothLodData).
-	 */
+	/** Returns \c true if \p InLodIndex is a valid LOD id (index into \c ClothLodData). */
 	virtual bool IsValidLod(int32 InLodIndex) override;
 
+	/** Returns the number of valid LOD's (length of the \c ClothLodData array). */
+	virtual int32 GetNumLods() override;
+
+#if WITH_EDITOR
+	/** * Add a new LOD class instance. */
+	virtual int32 AddNewLod() override;
+
 	/**
-	 * Returns the number of valid LOD's (length of the \c ClothLodData array).
+	 * Called on the clothing asset when the base data (physical mesh etc.) has
+	 * changed, so any intermediate generated data can be regenerated.
 	 */
-	virtual int32 GetNumLods() override;
+	virtual void InvalidateCachedData() override;
+
+	/* Called after changes in any of the asset properties. */
+	virtual void PostEditChangeChainProperty(FPropertyChangedChainEvent& ChainEvent) override;
+#endif // WITH_EDITOR
 
 	/**
 	 * Builds self collision data.
 	 * The default behavior is if the \c ClothSimConfig reports that self collisions
 	 * are enabled, it envokes \c BuildSelfCollisionData() on each \c ClothLodData's
-	 * \c PhysicalMeshData member.
+	 * \c ClothPhysicalMeshData member.
 	 */
 	virtual void BuildSelfCollisionData() override;
 
+	/** Return a const cloth config pointer of the desired cloth config type, or nullptr if there isn't any suitable. */
+	template<typename ClothConfigType, typename = typename TEnableIf<TIsDerivedFrom<ClothConfigType, UClothConfigBase>::IsDerived>::Type>
+	const ClothConfigType* GetClothConfig() const
+	{
+		const UClothConfigBase* const* const ClothConfig = ClothConfigs.Find(ClothConfigType::StaticClass()->GetFName());
+		return ClothConfig ? ExactCast<ClothConfigType>(*ClothConfig) : nullptr;
+	}
+
+	/** Return a cloth config pointer of the desired cloth config type, or nullptr if there isn't any suitable. */
+	template<typename ClothConfigType, typename = typename TEnableIf<TIsDerivedFrom<ClothConfigType, UClothConfigBase>::IsDerived>::Type>
+	ClothConfigType* GetClothConfig()
+	{
+		UClothConfigBase* const* const ClothConfig = ClothConfigs.Find(ClothConfigType::StaticClass()->GetFName());
+		return ClothConfig ? ExactCast<ClothConfigType>(*ClothConfig) : nullptr;
+	}
+
+	/** Add or replace a new cloth config of the specified type. */
+	template<typename ClothConfigType, typename = typename TEnableIf<TIsDerivedFrom<ClothConfigType, UClothConfigBase>::IsDerived>::Type>
+	void SetClothConfig(ClothConfigType* ClothConfig)
+	{
+		check(ClothConfig);
+		ClothConfigs.Add(ClothConfig->GetClass()->GetFName(), static_cast<UClothConfigBase*>(ClothConfig));
+	}
+
+	/** Migrate deprecated objects. */
+	virtual void PostLoad() override;
+
+	/** Serialize deprecated objects. */
+	virtual void Serialize(FArchive& Ar) override;
+
 	// The physics asset to extract collisions from when building a simulation.
 	UPROPERTY(EditAnywhere, Category = Config)
 	UPhysicsAsset* PhysicsAsset;
 
+	// Simulation specific cloth parameters. 
+	// Use GetClothConfig() to retrieve the correct parameters/config type for the desired cloth simulation system.
+	UPROPERTY(VisibleAnywhere, BlueprintReadOnly, EditFixedSize, Instanced, Category = Config)
+	TMap<FName, UClothConfigBase*> ClothConfigs;
+
 	// Shared by all cloth instances in a skeletal mesh
 	// Only supported with Chaos Cloth for now
 	// This may not be editable on unused cloth assets
-	UPROPERTY(EditAnywhere, Category = Config, NoClear, meta = (NoResetToDefault))
-	UClothSharedSimConfigBase* ClothSharedSimConfig;
+	UPROPERTY()
+	UClothConfigBase* ClothSharedSimConfig_DEPRECATED;
 
 	// Parameters for how the NVcloth behaves.
 	// These will have no effect on Chaos cloth
-	UPROPERTY(EditAnywhere, Category = Config, NoClear, meta = (NoResetToDefault))
-	UClothConfigBase* ClothSimConfig;
+	UPROPERTY()
+	UClothConfigBase* ClothSimConfig_DEPRECATED;
 
 	// Parameters for how Chaos cloth behaves 
 	// These will not affect NVcloth
 	// For now, we have two configuration parameters so that we can switch between chaos and
 	// non chaos at will without losing the original NVcloth data
-	UPROPERTY(EditAnywhere, Category = Config, NoClear, meta = (NoResetToDefault))
-	UClothConfigBase* ChaosClothSimConfig;
+	UPROPERTY()
+	UClothConfigBase* ChaosClothSimConfig_DEPRECATED;
 
 	// The actual asset data, listed by LOD.
 	UPROPERTY()
@@ -199,4 +241,25 @@ public:
 	// Custom data applied by the importer depending on where the asset was imported from.
 	UPROPERTY()
 	UClothingAssetCustomData* CustomData;
-};
\ No newline at end of file
+
+	/** 
+	 * Deprecated property for transitioning the \c FClothConfig struct to the 
+	 * \c UClothConfigBase array, in a new property called \c ClothConfigs.
+	 */
+	UPROPERTY()
+	FClothConfig_Legacy ClothConfig_DEPRECATED;
+
+	/** 
+	 * Deprecated property for transitioning \c FClothLODData struct to the 
+	 * \c UClothLODDataCommon class, in a new property called \c ClothLodData.
+	 */
+	UPROPERTY()
+	TArray<FClothLODData_Legacy> LodData_DEPRECATED;
+
+private:
+#if WITH_EDITOR
+	// Helper functions used in PostPropertyChangeCb
+	void ReregisterComponentsUsingClothing();
+	void ForEachInteractorUsingClothing(TFunction<void (UClothingSimulationInteractor*)> Func);
+#endif // WITH_EDITOR
+};
diff --git a/Engine/Source/Runtime/ClothingSystemRuntimeCommon/Public/PointWeightMap.h b/Engine/Source/Runtime/ClothingSystemRuntimeCommon/Public/PointWeightMap.h
new file mode 100644
--- /dev/null
+++ b/Engine/Source/Runtime/ClothingSystemRuntimeCommon/Public/PointWeightMap.h
@@ -0,0 +1,140 @@
+// Copyright 1998-2019 Epic Games, Inc. All Rights Reserved.
+#pragma once
+
+#include "UObject/ObjectMacros.h"
+#include "PointWeightMap.generated.h"
+
+/** The possible targets for a physical mesh point weight map. */
+UENUM()
+enum class EWeightMapTargetCommon : uint8
+{
+	None = 0,            // None, should always be zero
+	MaxDistance,         // The distance that each vertex can move away from its reference (skinned) position
+	BackstopDistance,    // Distance along the plane of the surface that the particles can travel (separation constraint)
+	BackstopRadius,      // Radius of movement to allow for backstop movement
+	AnimDriveMultiplier  // Strength of anim drive per-particle (spring driving particle back to skinned location
+};
+
+/** 
+ * A mask is simply some storage for a physical mesh parameter painted onto clothing.
+ * Used in the editor for users to paint onto and then target to a parameter, which
+ * is then later applied to a phys mesh
+ */
+USTRUCT()
+struct CLOTHINGSYSTEMRUNTIMECOMMON_API FPointWeightMap
+{
+	GENERATED_BODY();
+
+	FPointWeightMap()
+#if WITH_EDITORONLY_DATA
+		: Name(NAME_None)
+		, CurrentTarget((uint8)EWeightMapTargetCommon::None)
+		, bEnabled(false)
+#endif
+	{}
+
+	~FPointWeightMap()
+	{}
+
+	/**
+	 * Reset this map to the specified number of points, and set all the values to zero.
+	 * @param NumPoints the number of points to initialize the map with.
+	 */
+	void Initialize(const int32 NumPoints)
+	{
+		Values.Init(0.f, NumPoints);
+#if WITH_EDITORONLY_DATA
+		CurrentTarget = (uint8)EWeightMapTargetCommon::None;
+		bEnabled = false;
+#endif
+	}
+
+	/**
+	 * Initialize a weight map from another weight map while enabling and setting a new target.
+	 * @param Source the source weight map to copy the values from.
+	 * @param Target the new weight map target. */
+	template <typename T, typename = typename TEnableIf<TOr<TIsEnum<T>, TIsArithmetic<T>>::Value>::Type>
+	void Initialize(const FPointWeightMap& Source, T Target)
+	{
+		Values = Source.Values;
+#if WITH_EDITORONLY_DATA
+		CurrentTarget = (uint8)Target;
+		bEnabled = true;
+#endif
+	}
+
+	/** Empty this map of any values. */
+	void Empty()
+	{ Values.Empty(); }
+
+	/** Return the number of values in this map. */
+	int32 Num() const
+	{ return Values.Num(); }
+
+	/**
+	 * Return the current float value for the requested point.
+	 * @param Index the value/point index to retrieve, must be within range or it will assert.
+	 */
+	const float& operator[](int32 Index) const
+	{ return Values[Index]; }
+
+	/**
+	 * Return the current float value for the requested point.
+	 * @param Index the value/point index to retrieve, must be within range or it will assert.
+	 */
+	float& operator[](int32 Index)
+	{ return Values[Index]; }
+
+	/** 
+	 * Get a value from the map, or return 0 if the index is out of bounds.
+	 * @param Index the value/point index to retrieve
+	 */
+	float GetValue(int32 Index) const
+	{ return Values.IsValidIndex(Index) ? Values[Index] : 0.f; }
+
+	/** 
+	 * Set a value in the map checking first whether the index is within bounds.
+	 * @param Index the value/point index to set
+	 * @param Value the value to set
+	 */
+	void SetValue(int32 Index, float Value)
+	{ if (Values.IsValidIndex(Index)) Values[Index] = Value; }
+
+	/**
+	 * Return whether the specified point weight is below (or equal) to the specified threshold.
+	 * @param Index the value/point index to retrieve
+	 * @param Threshold, the value threshold to test against for.
+	 */
+	bool IsBelowThreshold(const int32 Index, const float Threshold=0.1f) const
+	{ return Values.IsValidIndex(Index) && Values[Index] <= Threshold; }
+
+	/** Return whether one of the specified triangle has at least one point weight below (or equal) to the specified @param Threshold. */
+	bool AreAnyBelowThreshold(const int32 Index0, const int32 Index1, const int32 Index2, const float Threshold=0.1f) const
+	{ return IsBelowThreshold(Index0, Threshold) || IsBelowThreshold(Index1, Threshold) || IsBelowThreshold(Index2, Threshold); }
+
+	/** Return whether all points' values are zero. */
+	bool IsZeroed() const
+	{ return !Values.FindByPredicate([](const float& Value) { return Value != 0.f; }); }
+
+	/** Calculates Min/Max values based on values. */
+	void CalcRanges(float& MinValue, float& MaxValue)
+	{ MinValue = FMath::Min(Values); MaxValue = FMath::Max(Values); }
+
+	/** The actual values stored in the mask */
+	UPROPERTY()
+	TArray<float> Values;
+
+#if WITH_EDITORONLY_DATA
+	/** Name of the mask, mainly for users to differentiate */
+	UPROPERTY()
+	FName Name;
+
+	/** The currently targeted parameter for the mask, @seealso EWeightMapTargetCommon */
+	UPROPERTY()
+	uint8 CurrentTarget;
+
+	/** Whether this mask is enabled and able to effect final mesh values */
+	UPROPERTY()
+	bool bEnabled;
+#endif
+};
diff --git a/Engine/Source/Runtime/ClothingSystemRuntimeCommon/Public/Utils/ClothingMeshUtils.h b/Engine/Source/Runtime/ClothingSystemRuntimeCommon/Public/Utils/ClothingMeshUtils.h
--- a/Engine/Source/Runtime/ClothingSystemRuntimeCommon/Public/Utils/ClothingMeshUtils.h
+++ b/Engine/Source/Runtime/ClothingSystemRuntimeCommon/Public/Utils/ClothingMeshUtils.h
@@ -6,7 +6,7 @@
 
 DECLARE_LOG_CATEGORY_EXTERN(LogClothingMeshUtils, Log, All);
 
-class UClothPhysicalMeshDataBase;
+struct FClothPhysicalMeshData;
 
 namespace ClothingMeshUtils
 {
@@ -35,7 +35,7 @@ namespace ClothingMeshUtils
 	void CLOTHINGSYSTEMRUNTIMECOMMON_API
 	SkinPhysicsMesh(
 		const TArray<int32>& BoneMap, // UClothingAssetCommon::UsedBoneIndices
-		const UClothPhysicalMeshDataBase& InMesh, 
+		const FClothPhysicalMeshData& InMesh, 
 		const FTransform& PostTransform, // Final transform to apply to component space positions and normals
 		const FMatrix* InBoneMatrices, 
 		const int32 InNumBoneMatrices, 
diff --git a/Engine/Source/Runtime/ClothingSystemRuntimeInterface/Private/ClothPhysicalMeshDataBase.cpp b/Engine/Source/Runtime/ClothingSystemRuntimeInterface/Private/ClothPhysicalMeshDataBase.cpp
deleted file mode 100644
diff --git a/Engine/Source/Runtime/ClothingSystemRuntimeInterface/Private/ClothPhysicalMeshDataBase_Legacy.cpp b/Engine/Source/Runtime/ClothingSystemRuntimeInterface/Private/ClothPhysicalMeshDataBase_Legacy.cpp
new file mode 100644
--- /dev/null
+++ b/Engine/Source/Runtime/ClothingSystemRuntimeInterface/Private/ClothPhysicalMeshDataBase_Legacy.cpp
@@ -0,0 +1,41 @@
+// Copyright 1998-2019 Epic Games, Inc. All Rights Reserved.
+
+#include "ClothPhysicalMeshDataBase_Legacy.h"
+
+UClothPhysicalMeshDataBase_Legacy::UClothPhysicalMeshDataBase_Legacy()
+	: NumFixedVerts(0)
+	, MaxBoneWeights(0)
+{}
+
+UClothPhysicalMeshDataBase_Legacy::~UClothPhysicalMeshDataBase_Legacy()
+{}
+
+TArray<float>* UClothPhysicalMeshDataBase_Legacy::GetFloatArray(const uint32 Id) const
+{
+	check(IdToArray.Contains(Id));
+	return IdToArray[Id];
+}
+
+TArray<uint32> UClothPhysicalMeshDataBase_Legacy::GetFloatArrayIds() const
+{
+	TArray<uint32> Keys; 
+	IdToArray.GetKeys(Keys);
+	return Keys;
+}
+
+TArray<TArray<float>*> UClothPhysicalMeshDataBase_Legacy::GetFloatArrays() const
+{
+	TArray<TArray<float>*> Values;
+	IdToArray.GenerateValueArray(Values);
+	return Values;
+}
+
+void UClothPhysicalMeshDataBase_Legacy::RegisterFloatArray(
+	const uint32 Id,
+	TArray<float> *Array)
+{
+	check(Id != INDEX_NONE);
+	check(Array != nullptr);
+	check(!IdToArray.Contains(Id) || IdToArray[Id] == Array);
+	IdToArray.Add(Id, Array);
+}
diff --git a/Engine/Source/Runtime/ClothingSystemRuntimeInterface/Public/ClothConfigBase.h b/Engine/Source/Runtime/ClothingSystemRuntimeInterface/Public/ClothConfigBase.h
--- a/Engine/Source/Runtime/ClothingSystemRuntimeInterface/Public/ClothConfigBase.h
+++ b/Engine/Source/Runtime/ClothingSystemRuntimeInterface/Public/ClothConfigBase.h
@@ -5,13 +5,11 @@
 #include "UObject/ObjectMacros.h"
 #include "ClothConfigBase.generated.h"
 
-class UClothingAssetBase;
-
 /**
  * Base class for simulator specific simulation controls.
  * Each cloth instance on a skeletal mesh can have a unique cloth config
  */
-UCLASS(Abstract, DefaultToInstanced)
+UCLASS(Abstract)
 class CLOTHINGSYSTEMRUNTIMEINTERFACE_API UClothConfigBase : public UObject
 {
 	GENERATED_BODY()
@@ -19,22 +17,35 @@ public:
 	UClothConfigBase();
 	virtual ~UClothConfigBase();
 
+	/**
+	 * Return whether self collision is enabled for this config.
+	 */
+	UE_DEPRECATED(4.25, "This function is deprecated. Please use NeedsSelfCollisionIndices instead.")
 	virtual bool HasSelfCollision() const
-	{
-		unimplemented(); return false;
-	}
-};
+	{ unimplemented(); return false; }
 
+	/**
+	 * Return the self collision radius if building self collision indices is required for this config.
+	 * Otherwise return 0.f.
+	 */
+	virtual float NeedsSelfCollisionIndices() const
+	{ return 0.f; }
+};
 
-/*
-These settings are shared between all instances on a skeletal mesh
-*/
-UCLASS(Abstract)
-class CLOTHINGSYSTEMRUNTIMEINTERFACE_API UClothSharedSimConfigBase : public UObject
+/**
+ * These settings are shared between all instances on a skeletal mesh
+ * Deprecated, use UClothConfigBase instead.
+ */
+UCLASS(Abstract, Deprecated)
+class CLOTHINGSYSTEMRUNTIMEINTERFACE_API UDEPRECATED_ClothSharedSimConfigBase : public UObject
 {
 	GENERATED_BODY()
 public:
-	UClothSharedSimConfigBase() {};
-	virtual ~UClothSharedSimConfigBase() {};
-	
+	UDEPRECATED_ClothSharedSimConfigBase() {}
+	virtual ~UDEPRECATED_ClothSharedSimConfigBase() {}
+
+	/**
+	 * Return a new updated cloth shared config migrated from this current object.
+	 */
+	virtual UClothConfigBase* Migrate() { return nullptr; }
 };
diff --git a/Engine/Source/Runtime/ClothingSystemRuntimeInterface/Public/ClothPhysicalMeshDataBase.h b/Engine/Source/Runtime/ClothingSystemRuntimeInterface/Public/ClothPhysicalMeshDataBase.h
deleted file mode 100644
diff --git a/Engine/Source/Runtime/ClothingSystemRuntimeInterface/Public/ClothPhysicalMeshDataBase_Legacy.h b/Engine/Source/Runtime/ClothingSystemRuntimeInterface/Public/ClothPhysicalMeshDataBase_Legacy.h
new file mode 100644
--- /dev/null
+++ b/Engine/Source/Runtime/ClothingSystemRuntimeInterface/Public/ClothPhysicalMeshDataBase_Legacy.h
@@ -0,0 +1,81 @@
+// Copyright 1998-2019 Epic Games, Inc. All Rights Reserved.
+#pragma once
+
+#include "UObject/Object.h"
+
+#include "ClothVertBoneData.h"
+#include "Containers/Array.h"
+#include "Containers/Map.h"
+
+#include "ClothPhysicalMeshDataBase_Legacy.generated.h"
+
+/**
+ * Deprecated, use FClothPhysicalMeshData instead.
+ * Simulation mesh points, topology, and spatial parameters defined on that 
+ * topology.
+ *
+ * Created curing asset import or created from a skeletal mesh.
+ */
+UCLASS()
+class CLOTHINGSYSTEMRUNTIMEINTERFACE_API UClothPhysicalMeshDataBase_Legacy : public UObject
+{
+	GENERATED_BODY()
+public:
+	UClothPhysicalMeshDataBase_Legacy();
+	virtual ~UClothPhysicalMeshDataBase_Legacy();
+
+	/** Retrieve a registered vertex weight array by unique @param Id. */
+	TArray<float>* GetFloatArray(const uint32 Id) const;
+
+	/** Get ids for all registered weight arrays. */
+	TArray<uint32> GetFloatArrayIds() const;
+
+	/** Get all registered weight arrays. */
+	TArray<TArray<float>*> GetFloatArrays() const;
+
+protected:
+	/** Register an @param Array keyed by a unique @param Id. */
+	void RegisterFloatArray(const uint32 Id, TArray<float> *Array);
+
+public:
+	// Positions of each simulation vertex
+	UPROPERTY()
+	TArray<FVector> Vertices;
+
+	// Normal at each vertex
+	UPROPERTY()
+	TArray<FVector> Normals;
+
+#if WITH_EDITORONLY_DATA
+	// Color at each vertex
+	UPROPERTY()
+	TArray<FColor> VertexColors;
+#endif // WITH_EDITORONLY_DATA
+
+	// Indices of the simulation mesh triangles
+	UPROPERTY()
+	TArray<uint32> Indices;
+
+	// Inverse mass for each vertex in the physical mesh
+	UPROPERTY()
+	TArray<float> InverseMasses;
+
+	// Indices and weights for each vertex, used to skin the mesh to create the reference pose
+	UPROPERTY()
+	TArray<FClothVertBoneData> BoneData;
+
+	// Number of fixed verts in the simulation mesh (fixed verts are just skinned and do not simulate)
+	UPROPERTY()
+	int32 NumFixedVerts;
+
+	// Maximum number of bone weights of any vetex
+	UPROPERTY()
+	int32 MaxBoneWeights;
+
+	// Valid indices to use for self collisions (reduced set of Indices)
+	UPROPERTY()
+	TArray<uint32> SelfCollisionIndices;
+
+private:
+	TMap<uint32, TArray<float>*> IdToArray;
+};
diff --git a/Engine/Source/Runtime/ClothingSystemRuntimeInterface/Public/ClothingSimulationFactory.h b/Engine/Source/Runtime/ClothingSystemRuntimeInterface/Public/ClothingSimulationFactory.h
--- a/Engine/Source/Runtime/ClothingSystemRuntimeInterface/Public/ClothingSimulationFactory.h
+++ b/Engine/Source/Runtime/ClothingSystemRuntimeInterface/Public/ClothingSimulationFactory.h
@@ -68,14 +68,13 @@ public:
 	// Creates the runtime interactor object for a clothing simulation. This object will
 	// receive events allowing it to write data to the simulation context in a safe manner
 	virtual UClothingSimulationInteractor* CreateInteractor()
-	{
-		return nullptr;
-	}
-
-	// Creates a default cloth configuration during construction of the clothing Asset
-	// TBC: This may be removed after we have a specialized ChaosClothingAsset
-	virtual UClothConfigBase* CreateDefaultClothConfig(const FObjectInitializer& ObjectInitializer, UObject* Outer)
-	{
-		return nullptr;
-	}
+	PURE_VIRTUAL(UClothingSimulationFactory::CreateInteractor, return nullptr;);
+
+	// Return the cloth config type for this cloth factory
+	virtual TSubclassOf<UClothConfigBase> GetClothConfigClass() const
+	PURE_VIRTUAL(UClothingSimulationFactory::GetClothConfigClass, return nullptr;);
+
+	// Return an enum of the weight map targets that can be used with this simulation.
+	virtual const UEnum* GetWeightMapTargetEnum() const
+	PURE_VIRTUAL(UClothingSimulationFactory::GetWeightMapTargetEnum, return nullptr;);
 };
diff --git a/Engine/Source/Runtime/ClothingSystemRuntimeInterface/Public/PointWeightMap.h b/Engine/Source/Runtime/ClothingSystemRuntimeInterface/Public/PointWeightMap.h
deleted file mode 100644
diff --git a/Engine/Source/Runtime/ClothingSystemRuntimeNv/Private/Assets/ClothingAssetNv.cpp b/Engine/Source/Runtime/ClothingSystemRuntimeNv/Private/Assets/ClothingAssetNv.cpp
deleted file mode 100644
diff --git a/Engine/Source/Runtime/ClothingSystemRuntimeNv/Private/ClothConfigNv.cpp b/Engine/Source/Runtime/ClothingSystemRuntimeNv/Private/ClothConfigNv.cpp
--- a/Engine/Source/Runtime/ClothingSystemRuntimeNv/Private/ClothConfigNv.cpp
+++ b/Engine/Source/Runtime/ClothingSystemRuntimeNv/Private/ClothConfigNv.cpp
@@ -1,9 +1,26 @@
 // Copyright 1998-2019 Epic Games, Inc. All Rights Reserved.
 
 #include "ClothConfigNv.h"
+#include "ClothConfig_Legacy.h"
+#include "ClothConfigNvCustomVersion.h"
+
+FClothConstraintSetupNv::FClothConstraintSetupNv()
+	: Stiffness(1.0f)
+	, StiffnessMultiplier(1.0f)
+	, StretchLimit(1.0f)
+	, CompressionLimit(1.0f)
+{}
+
+void FClothConstraintSetupNv::MigrateFrom(const FClothConstraintSetup_Legacy& Setup)
+{
+	Stiffness = Setup.Stiffness;
+	StiffnessMultiplier = Setup.StiffnessMultiplier;
+	StretchLimit = Setup.StretchLimit;
+	CompressionLimit = Setup.CompressionLimit;
+}
 
 UClothConfigNv::UClothConfigNv()
-	: WindMethod(EClothingWindMethod::Legacy)
+	: ClothingWindMethod(EClothingWindMethodNv::Legacy)
 	, SelfCollisionRadius(0.0f)
 	, SelfCollisionStiffness(0.0f)
 	, SelfCollisionCullScale(1.0f)
@@ -28,22 +45,65 @@ UClothConfigNv::UClothConfigNv()
 	, AnimDriveDamperStiffness(1.0f)
 {}
 
-bool UClothConfigNv::HasSelfCollision() const
+void UClothConfigNv::Serialize(FArchive& Ar)
 {
-	return SelfCollisionRadius > 0.0f && SelfCollisionStiffness > 0.0f;
+	Super::Serialize(Ar);
+	Ar.UsingCustomVersion(FClothConfigNvCustomVersion::GUID);
+}
+
+void UClothConfigNv::PostLoad()
+{
+	Super::PostLoad();
+
+	const int32 ClothingCustomVersion = GetLinkerCustomVersion(FClothConfigNvCustomVersion::GUID);
+	
+	if (ClothingCustomVersion < FClothConfigNvCustomVersion::DeprecateLegacyStructureAndEnum)
+	{
+		ClothingWindMethod = static_cast<EClothingWindMethodNv>(WindMethod_DEPRECATED);
+		VerticalConstraint.MigrateFrom(VerticalConstraintConfig_DEPRECATED);
+		HorizontalConstraint.MigrateFrom(HorizontalConstraintConfig_DEPRECATED);
+		BendConstraint.MigrateFrom(BendConstraintConfig_DEPRECATED);
+		ShearConstraint.MigrateFrom(ShearConstraintConfig_DEPRECATED);
+	}
+}
+
+void UClothConfigNv::MigrateFrom(const FClothConfig_Legacy& ClothConfig)
+{
+	ClothingWindMethod = static_cast<EClothingWindMethodNv>(ClothConfig.WindMethod);
+	VerticalConstraint.MigrateFrom(ClothConfig.VerticalConstraintConfig);
+	HorizontalConstraint.MigrateFrom(ClothConfig.HorizontalConstraintConfig);
+	BendConstraint.MigrateFrom(ClothConfig.BendConstraintConfig);
+	ShearConstraint.MigrateFrom(ClothConfig.ShearConstraintConfig);
+	SelfCollisionRadius = ClothConfig.SelfCollisionRadius;
+	SelfCollisionStiffness = ClothConfig.SelfCollisionStiffness;
+	SelfCollisionCullScale = ClothConfig.SelfCollisionCullScale;
+	Damping = ClothConfig.Damping;
+	Friction = ClothConfig.Friction;
+	WindDragCoefficient = ClothConfig.WindDragCoefficient;
+	WindLiftCoefficient = ClothConfig.WindLiftCoefficient;
+	LinearDrag = ClothConfig.LinearDrag;
+	AngularDrag = ClothConfig.AngularDrag;
+	LinearInertiaScale = ClothConfig.LinearInertiaScale;
+	AngularInertiaScale = ClothConfig.AngularInertiaScale;
+	CentrifugalInertiaScale = ClothConfig.CentrifugalInertiaScale;
+	SolverFrequency = ClothConfig.SolverFrequency;
+	StiffnessFrequency = ClothConfig.StiffnessFrequency;
+	GravityScale = ClothConfig.GravityScale;
+	GravityOverride = ClothConfig.GravityOverride;
+	bUseGravityOverride = ClothConfig.bUseGravityOverride;
+	TetherStiffness = ClothConfig.TetherStiffness;
+	TetherLimit = ClothConfig.TetherLimit;
+	CollisionThickness = ClothConfig.CollisionThickness;
+	AnimDriveSpringStiffness = ClothConfig.AnimDriveSpringStiffness;
+	AnimDriveDamperStiffness = ClothConfig.AnimDriveDamperStiffness;
 }
 
-#if WITH_EDITOR
-bool UClothConfigNv::InitFromApexAssetCallback(
-	nvidia::apex::ClothingAsset* InApexAsset,
-	USkeletalMesh* TargetMesh,
-	FName InName) 
+bool UClothConfigNv::UseSelfCollisions() const
 {
-#if WITH_APEX_CLOTHING
-	// Set to use legacy wind calculations, which is what APEX would normally have used
-	WindMethod = EClothingWindMethod::Legacy;
-#endif // WITH_APEX_CLOTHING
-	return false;
+	return SelfCollisionRadius > 0.0f && SelfCollisionStiffness > 0.0f;
 }
 
-#endif //WITH_EDITOR
+float UClothConfigNv::NeedsSelfCollisionIndices() const
+{
+	return UseSelfCollisions() ? SelfCollisionRadius * SelfCollisionCullScale : 0.0f;
+}
diff --git a/Engine/Source/Runtime/ClothingSystemRuntimeNv/Private/ClothConfigNvCustomVersion.cpp b/Engine/Source/Runtime/ClothingSystemRuntimeNv/Private/ClothConfigNvCustomVersion.cpp
new file mode 100644
--- /dev/null
+++ b/Engine/Source/Runtime/ClothingSystemRuntimeNv/Private/ClothConfigNvCustomVersion.cpp
@@ -0,0 +1,9 @@
+// Copyright 1998-2019 Epic Games, Inc. All Rights Reserved.
+
+#include "ClothConfigNvCustomVersion.h"
+#include "Serialization/CustomVersion.h"
+
+const FGuid FClothConfigNvCustomVersion::GUID(0x9950B70E, 0xB41A4E17, 0xBBCCFA0D, 0x57817FD6);
+
+// Register the custom version with core
+FCustomVersionRegistration GRegisterClothConfigNvCustomVersion(FClothConfigNvCustomVersion::GUID, FClothConfigNvCustomVersion::LatestVersion, TEXT("ClothConfigNvVer"));
diff --git a/Engine/Source/Runtime/ClothingSystemRuntimeNv/Private/ClothConfigNvCustomVersion.h b/Engine/Source/Runtime/ClothingSystemRuntimeNv/Private/ClothConfigNvCustomVersion.h
new file mode 100644
--- /dev/null
+++ b/Engine/Source/Runtime/ClothingSystemRuntimeNv/Private/ClothConfigNvCustomVersion.h
@@ -0,0 +1,28 @@
+// Copyright 1998-2019 Epic Games, Inc. All Rights Reserved.
+
+#pragma once
+
+#include "CoreMinimal.h"
+#include "Misc/Guid.h"
+
+// Custom serialization version for config
+struct FClothConfigNvCustomVersion
+{
+	enum Type
+	{
+		// Before any version changes were made
+		BeforeCustomVersionWasAdded = 0,
+		// Deprecate legacy structure and enum that couldn't be redirected
+		DeprecateLegacyStructureAndEnum = 1,
+
+		// -----<new versions can be added above this line>-------------------------------------------------
+		VersionPlusOne,
+		LatestVersion = VersionPlusOne - 1
+	};
+
+	// The GUID for this custom version number
+	const static FGuid GUID;
+
+private:
+	FClothConfigNvCustomVersion() {}
+};
diff --git a/Engine/Source/Runtime/ClothingSystemRuntimeNv/Private/ClothLODDataNv.cpp b/Engine/Source/Runtime/ClothingSystemRuntimeNv/Private/ClothLODDataNv.cpp
deleted file mode 100644
diff --git a/Engine/Source/Runtime/ClothingSystemRuntimeNv/Private/ClothPhysicalMeshDataNv.cpp b/Engine/Source/Runtime/ClothingSystemRuntimeNv/Private/ClothPhysicalMeshDataNv.cpp
deleted file mode 100644
diff --git a/Engine/Source/Runtime/ClothingSystemRuntimeNv/Private/ClothPhysicalMeshDataNv_Legacy.cpp b/Engine/Source/Runtime/ClothingSystemRuntimeNv/Private/ClothPhysicalMeshDataNv_Legacy.cpp
new file mode 100644
--- /dev/null
+++ b/Engine/Source/Runtime/ClothingSystemRuntimeNv/Private/ClothPhysicalMeshDataNv_Legacy.cpp
@@ -0,0 +1,14 @@
+// Copyright 1998-2019 Epic Games, Inc. All Rights Reserved.
+#include "ClothPhysicalMeshDataNv_Legacy.h"
+#include "ClothPhysicalMeshData.h"  // For EWeightMapTargetCommon
+
+UClothPhysicalMeshDataNv_Legacy::UClothPhysicalMeshDataNv_Legacy()
+{
+	Super::RegisterFloatArray((uint32)EWeightMapTargetCommon::MaxDistance, &MaxDistances);
+	Super::RegisterFloatArray((uint32)EWeightMapTargetCommon::BackstopDistance, &BackstopDistances);
+	Super::RegisterFloatArray((uint32)EWeightMapTargetCommon::BackstopRadius, &BackstopRadiuses);
+	Super::RegisterFloatArray((uint32)EWeightMapTargetCommon::AnimDriveMultiplier, &AnimDriveMultipliers);
+}
+
+UClothPhysicalMeshDataNv_Legacy::~UClothPhysicalMeshDataNv_Legacy()
+{}
diff --git a/Engine/Source/Runtime/ClothingSystemRuntimeNv/Private/ClothingSimulationFactoryNv.cpp b/Engine/Source/Runtime/ClothingSystemRuntimeNv/Private/ClothingSimulationFactoryNv.cpp
--- a/Engine/Source/Runtime/ClothingSystemRuntimeNv/Private/ClothingSimulationFactoryNv.cpp
+++ b/Engine/Source/Runtime/ClothingSystemRuntimeNv/Private/ClothingSimulationFactoryNv.cpp
@@ -1,22 +1,20 @@
 // Copyright 1998-2019 Epic Games, Inc. All Rights Reserved.
 
 #include "ClothingSimulationFactoryNv.h"
-#include "Assets/ClothingAssetNv.h"
 
 #if WITH_NVCLOTH
+#include "ClothPhysicalMeshData.h"  // For EWeightMapTargetCommon
 #include "ClothingSimulationNv.h"
-#endif
-
 #include "ClothingSimulationInteractorNv.h"
-#include "UObject/Package.h"
+#endif
 
 IClothingSimulation* UClothingSimulationFactoryNv::CreateSimulation()
 {
 #if WITH_NVCLOTH
-	FClothingSimulationBase* Simulation = new FClothingSimulationNv();
-	return Simulation;
-#endif
+	return new FClothingSimulationNv();
+#else
 	return nullptr;
+#endif
 }
 
 void UClothingSimulationFactoryNv::DestroySimulation(IClothingSimulation* InSimulation)
@@ -30,8 +28,9 @@ bool UClothingSimulationFactoryNv::SupportsAsset(UClothingAssetBase* InAsset)
 {
 #if WITH_NVCLOTH
 	return true;
-#endif
+#else
 	return false;
+#endif
 }
 
 bool UClothingSimulationFactoryNv::SupportsRuntimeInteraction()
@@ -41,5 +40,27 @@ bool UClothingSimulationFactoryNv::SupportsRuntimeInteraction()
 
 UClothingSimulationInteractor* UClothingSimulationFactoryNv::CreateInteractor()
 {
+#if WITH_NVCLOTH
 	return NewObject<UClothingSimulationInteractorNv>(GetTransientPackage());
+#else
+	return nullptr;
+#endif
+}
+
+TSubclassOf<UClothConfigBase> UClothingSimulationFactoryNv::GetClothConfigClass() const
+{
+#if WITH_NVCLOTH
+	return TSubclassOf<UClothConfigBase>(UClothConfigNv::StaticClass());
+#else
+	return nullptr;
+#endif
+}
+
+const UEnum* UClothingSimulationFactoryNv::GetWeightMapTargetEnum() const
+{
+#if WITH_NVCLOTH
+	return StaticEnum<EWeightMapTargetCommon>();
+#else
+	return nullptr;
+#endif
 }
diff --git a/Engine/Source/Runtime/ClothingSystemRuntimeNv/Private/ClothingSimulationNv.cpp b/Engine/Source/Runtime/ClothingSystemRuntimeNv/Private/ClothingSimulationNv.cpp
--- a/Engine/Source/Runtime/ClothingSystemRuntimeNv/Private/ClothingSimulationNv.cpp
+++ b/Engine/Source/Runtime/ClothingSystemRuntimeNv/Private/ClothingSimulationNv.cpp
@@ -4,10 +4,11 @@
 
 #if WITH_NVCLOTH
 
-#include "Assets/ClothingAssetNv.h"
 #include "ClothingSystemRuntimeModuleNv.h"
 #include "NvClothIncludes.h"
 #include "NvClothSupport.h"
+#include "ClothPhysicalMeshData.h"
+#include "ClothingAsset.h"
 
 #include "Utils/ClothingMeshUtils.h"
 
@@ -216,12 +217,7 @@ void FClothingSimulationNv::CreateActor(USkeletalMeshComponent* InOwnerComponent
 	for(int32 LodIndex = 0; LodIndex < Asset->ClothLodData.Num(); ++LodIndex)
 	{
 		UClothLODDataCommon& AssetLodData = *Asset->ClothLodData[LodIndex];
-		UClothPhysicalMeshDataNv* PhysMeshPtr = Cast<UClothPhysicalMeshDataNv>(AssetLodData.PhysicalMeshData);
-		if (!ensure(PhysMeshPtr))
-		{
-			continue;
-		}
-		UClothPhysicalMeshDataNv& PhysMesh = *PhysMeshPtr;
+		const FClothPhysicalMeshData& PhysMesh = AssetLodData.ClothPhysicalMeshData;
 
 		NewActor.LodData.AddDefaulted();
 		FClothingActorNv::FActorLodData& ActorLodData = NewActor.LodData.Last();
@@ -332,11 +328,11 @@ void FClothingSimulationNv::CreateActor(USkeletalMeshComponent* InOwnerComponent
 		nv::cloth::Range<physx::PxVec4> MotionConstraints = NewCloth->getMotionConstraints();
 		const int32 NumMotionConstraints = NewCloth->getNumMotionConstraints();
 		check(NumMotionConstraints == Verts.Num());
+		const FPointWeightMap& MaxDistances = PhysMesh.GetWeightMap(EWeightMapTargetCommon::MaxDistance);
 		for(int32 ConstraintIndex = 0; ConstraintIndex < NumMotionConstraints; ++ConstraintIndex)
 		{
 			physx::PxVec4& Constraint = MotionConstraints[ConstraintIndex];
-
-			Constraint = physx::PxVec4(U2PVector(SkinnedVerts[ConstraintIndex]), PhysMesh.MaxDistances[ConstraintIndex]);
+			Constraint = physx::PxVec4(U2PVector(SkinnedVerts[ConstraintIndex]), MaxDistances[ConstraintIndex]);
 		}
 
 		// Set up the starting transform data for the cloth, then clear our inertia so
@@ -364,7 +360,11 @@ void FClothingSimulationNv::CreateActor(USkeletalMeshComponent* InOwnerComponent
 		}
 	}
 
-	ApplyClothConfig(Asset->ClothSimConfig, NewActor);
+	const UClothConfigNv* const Config = Asset->GetClothConfig<UClothConfigNv>();
+	if (ensure(Config != nullptr))
+	{
+		ApplyClothConfig(Config, NewActor);
+	}
 
 	// Pull collisions from the specified physics asset inside the clothing asset
 	ExtractActorCollisions(Asset, NewActor);
@@ -474,26 +474,21 @@ void FClothingSimulationNv::ExtractActorCollisions(UClothingAssetCommon* Asset,
 	}
 }
 
-void FClothingSimulationNv::ApplyClothConfig(UClothConfigBase *BaseConfig, FClothingActorNv &InActor)
+void FClothingSimulationNv::ApplyClothConfig(const UClothConfigNv* Config, FClothingActorNv &InActor)
 {
-	check(BaseConfig);
-	UClothConfigNv* ConfigPtr = Cast<UClothConfigNv>(BaseConfig);
-	check(ConfigPtr);
-	UClothConfigNv& Config = *ConfigPtr;
-
 	// These calculations convert our damping and drag values to a value closer to the way APEX used them.
 	// This allows us to not break old assets, and expectations from the old system
 	const float DampStiffnesssFreq = 10.0f;
-	const float DampStiffFreqRatio = DampStiffnesssFreq / Config.StiffnessFrequency;
-	const float ExpDampX = DampStiffFreqRatio * FMath::Log2(1.0f - Config.Damping.X);
-	const float ExpDampY = DampStiffFreqRatio * FMath::Log2(1.0f - Config.Damping.Y);
-	const float ExpDampZ = DampStiffFreqRatio * FMath::Log2(1.0f - Config.Damping.Z);
-	const float ExpDragLinX = DampStiffFreqRatio * FMath::Log2(1.0f - Config.LinearDrag.X);
-	const float ExpDragLinY = DampStiffFreqRatio * FMath::Log2(1.0f - Config.LinearDrag.Y);
-	const float ExpDragLinZ = DampStiffFreqRatio * FMath::Log2(1.0f - Config.LinearDrag.Z);
-	const float ExpDragAngX = DampStiffFreqRatio * FMath::Log2(1.0f - Config.AngularDrag.X);
-	const float ExpDragAngY = DampStiffFreqRatio * FMath::Log2(1.0f - Config.AngularDrag.Y);
-	const float ExpDragAngZ = DampStiffFreqRatio * FMath::Log2(1.0f - Config.AngularDrag.Z);
+	const float DampStiffFreqRatio = DampStiffnesssFreq / Config->StiffnessFrequency;
+	const float ExpDampX = DampStiffFreqRatio * FMath::Log2(1.0f - Config->Damping.X);
+	const float ExpDampY = DampStiffFreqRatio * FMath::Log2(1.0f - Config->Damping.Y);
+	const float ExpDampZ = DampStiffFreqRatio * FMath::Log2(1.0f - Config->Damping.Z);
+	const float ExpDragLinX = DampStiffFreqRatio * FMath::Log2(1.0f - Config->LinearDrag.X);
+	const float ExpDragLinY = DampStiffFreqRatio * FMath::Log2(1.0f - Config->LinearDrag.Y);
+	const float ExpDragLinZ = DampStiffFreqRatio * FMath::Log2(1.0f - Config->LinearDrag.Z);
+	const float ExpDragAngX = DampStiffFreqRatio * FMath::Log2(1.0f - Config->AngularDrag.X);
+	const float ExpDragAngY = DampStiffFreqRatio * FMath::Log2(1.0f - Config->AngularDrag.Y);
+	const float ExpDragAngZ = DampStiffFreqRatio * FMath::Log2(1.0f - Config->AngularDrag.Z);
 
 	const float PrecalcLog2 = FMath::Loge(2.0f);
 	const FVector AdjustedDamping(1.0f - FMath::Exp(ExpDampX * PrecalcLog2),
@@ -520,7 +515,7 @@ void FClothingSimulationNv::ApplyClothConfig(UClothConfigBase *BaseConfig, FClot
 			nv::cloth::PhaseConfig& PhaseConfig = LodData.PhaseConfigs[PhaseIndex];
 			const nv::cloth::ClothFabricPhaseType::Enum PhaseType = (nv::cloth::ClothFabricPhaseType::Enum)LodData.PhaseTypes[PhaseIndex];
 
-			FClothConstraintSetup* ConfigToUse = nullptr;
+			const FClothConstraintSetupNv* ConfigToUse = nullptr;
 
 			switch(PhaseType)
 			{
@@ -528,16 +523,16 @@ void FClothingSimulationNv::ApplyClothConfig(UClothConfigBase *BaseConfig, FClot
 					check(false);
 					break;
 				case nv::cloth::ClothFabricPhaseType::eVERTICAL:
-					ConfigToUse = &Config.VerticalConstraintConfig;
+					ConfigToUse = &Config->VerticalConstraint;
 					break;
 				case nv::cloth::ClothFabricPhaseType::eHORIZONTAL:
-					ConfigToUse = &Config.HorizontalConstraintConfig;
+					ConfigToUse = &Config->HorizontalConstraint;
 					break;
 				case nv::cloth::ClothFabricPhaseType::eBENDING:
-					ConfigToUse = &Config.BendConstraintConfig;
+					ConfigToUse = &Config->BendConstraint;
 					break;
 				case nv::cloth::ClothFabricPhaseType::eSHEARING:
-					ConfigToUse = &Config.ShearConstraintConfig;
+					ConfigToUse = &Config->ShearConstraint;
 					break;
 			}
 
@@ -551,10 +546,10 @@ void FClothingSimulationNv::ApplyClothConfig(UClothConfigBase *BaseConfig, FClot
 		LodCloth->setPhaseConfig(NvClothSupport::CreateRange(LodData.PhaseConfigs));
 
 		// Drag and lift coeffs only take effect in accurate wind mode
-		if(Config.WindMethod == EClothingWindMethod::Accurate)
+		if(Config->ClothingWindMethod == EClothingWindMethodNv::Accurate)
 		{
-			LodCloth->setDragCoefficient(Config.WindDragCoefficient);
-			LodCloth->setLiftCoefficient(Config.WindLiftCoefficient);
+			LodCloth->setDragCoefficient(Config->WindDragCoefficient);
+			LodCloth->setLiftCoefficient(Config->WindLiftCoefficient);
 		}
 		else
 		{
@@ -562,41 +557,41 @@ void FClothingSimulationNv::ApplyClothConfig(UClothConfigBase *BaseConfig, FClot
 			LodCloth->setLiftCoefficient(0.0f);
 		}
 
-		if(InActor.WindMethod == EClothingWindMethod::Legacy && Config.WindMethod == EClothingWindMethod::Accurate)
+		if(InActor.WindMethod == EClothingWindMethodNv::Legacy && Config->ClothingWindMethod == EClothingWindMethodNv::Accurate)
 		{
 			// Need to clear out particle accelerations here otherwise legacy wind effects will stay applied
 			LodCloth->clearParticleAccelerations();
 		}
 
-		LodCloth->setSolverFrequency(Config.SolverFrequency);
-		LodCloth->setStiffnessFrequency(Config.StiffnessFrequency);
-		LodCloth->setAcceleationFilterWidth(2.0f * (uint32)Config.SolverFrequency);
+		LodCloth->setSolverFrequency(Config->SolverFrequency);
+		LodCloth->setStiffnessFrequency(Config->StiffnessFrequency);
+		LodCloth->setAcceleationFilterWidth(2.0f * (uint32)Config->SolverFrequency);
 
-		LodCloth->setSelfCollisionDistance(Config.SelfCollisionRadius);
-		LodCloth->setSelfCollisionStiffness(Config.SelfCollisionStiffness);
+		LodCloth->setSelfCollisionDistance(Config->SelfCollisionRadius);
+		LodCloth->setSelfCollisionStiffness(Config->SelfCollisionStiffness);
 
 		// If we have self collision, we need to set rest positions for the mesh
-		if(Config.HasSelfCollision())
+		if(Config->UseSelfCollisions())
 		{
 			LodCloth->setRestPositions(NvClothSupport::CreateRange(LodData.Px_RestPositions));
 		}
 
 		LodCloth->setDamping(U2PVector(AdjustedDamping));
-		LodCloth->setFriction(Config.Friction);
+		LodCloth->setFriction(Config->Friction);
 		LodCloth->setLinearDrag(U2PVector(AdjustedDragLin));
 		LodCloth->setAngularDrag(U2PVector(AdjustedDragAng));
-		LodCloth->setLinearInertia(U2PVector(Config.LinearInertiaScale));
-		LodCloth->setAngularInertia(U2PVector(Config.AngularInertiaScale));
-		LodCloth->setCentrifugalInertia(U2PVector(Config.CentrifugalInertiaScale));
+		LodCloth->setLinearInertia(U2PVector(Config->LinearInertiaScale));
+		LodCloth->setAngularInertia(U2PVector(Config->AngularInertiaScale));
+		LodCloth->setCentrifugalInertia(U2PVector(Config->CentrifugalInertiaScale));
 
-		LodCloth->setTetherConstraintScale(Config.TetherLimit);
-		LodCloth->setTetherConstraintStiffness(Config.TetherStiffness);
+		LodCloth->setTetherConstraintScale(Config->TetherLimit);
+		LodCloth->setTetherConstraintStiffness(Config->TetherStiffness);
 	}
 
-	InActor.CollisionThickness = Config.CollisionThickness;
-	InActor.WindMethod = Config.WindMethod;
-	InActor.CurrentAnimDriveSpringStiffness = Config.AnimDriveSpringStiffness;
-	InActor.CurrentAnimDriveDamperStiffness = Config.AnimDriveDamperStiffness;
+	InActor.CollisionThickness = Config->CollisionThickness;
+	InActor.WindMethod = Config->ClothingWindMethod;
+	InActor.CurrentAnimDriveSpringStiffness = Config->AnimDriveSpringStiffness;
+	InActor.CurrentAnimDriveDamperStiffness = Config->AnimDriveDamperStiffness;
 }
 
 IClothingSimulationContext* FClothingSimulationNv::CreateContext()
@@ -678,7 +673,7 @@ void FClothingSimulationNv::Simulate(IClothingSimulationContext* InContext)
 
 		// To build motion constraints (max distances) we need to skin the entire physics mesh
 		// this call also updates our fixed particles to avoid iterating the particle list a second time
-		const UClothPhysicalMeshDataBase& PhysMesh = *Actor.AssetCreatedFrom->ClothLodData[Actor.CurrentLodIndex]->PhysicalMeshData;
+		const FClothPhysicalMeshData& PhysMesh = Actor.AssetCreatedFrom->ClothLodData[Actor.CurrentLodIndex]->ClothPhysicalMeshData;
 		ClothingMeshUtils::SkinPhysicsMesh(Actor.AssetCreatedFrom->UsedBoneIndices, PhysMesh, RootBoneTransform, NvContext->RefToLocals.GetData(), NvContext->RefToLocals.Num(), Actor.SkinnedPhysicsMeshPositions[Actor.CurrentSkinnedPositionIndex], Actor.SkinnedPhysicsMeshNormals);
 
 		nv::cloth::Cloth* CurrentCloth = Actor.LodData[Actor.CurrentLodIndex].Cloth;
@@ -717,22 +712,21 @@ void FClothingSimulationNv::Simulate(IClothingSimulationContext* InContext)
 		CurrentCloth->setTranslation(U2PVector(RootBoneWorldTransform.GetTranslation()));
 		CurrentCloth->setRotation(U2PQuat(RootBoneWorldTransform.GetRotation()));
 
-		UClothConfigBase* ClothConfigBase = Actor.AssetCreatedFrom->ClothSimConfig;
-		check(ClothConfigBase);
-		UClothConfigNv* ConfigPtr = Cast<UClothConfigNv>(ClothConfigBase);
-		check(ConfigPtr);
-		UClothConfigNv& ClothConfig = *ConfigPtr;
-		if (ClothConfig.bUseGravityOverride)
+		if (const UClothConfigNv* const ConfigPtr = Actor.AssetCreatedFrom->GetClothConfig<UClothConfigNv>())
 		{
-			CurrentCloth->setGravity(U2PVector(ClothConfig.GravityOverride));
-		}
-		else if(Actor.bUseGravityOverride)
-		{
-			CurrentCloth->setGravity(U2PVector(ClothConfig.GravityScale * Actor.GravityOverride));
-		}
-		else
-		{
-			CurrentCloth->setGravity(U2PVector(ClothConfig.GravityScale * NvContext->WorldGravity));
+			const UClothConfigNv& ClothConfig = *ConfigPtr;
+			if (ClothConfig.bUseGravityOverride)
+			{
+				CurrentCloth->setGravity(U2PVector(ClothConfig.GravityOverride));
+			}
+			else if(Actor.bUseGravityOverride)
+			{
+				CurrentCloth->setGravity(U2PVector(ClothConfig.GravityScale * Actor.GravityOverride));
+			}
+			else
+			{
+				CurrentCloth->setGravity(U2PVector(ClothConfig.GravityScale * NvContext->WorldGravity));
+			}
 		}
 
 		Actor.UpdateMotionConstraints(NvContext);
@@ -867,7 +861,7 @@ void FClothingSimulationNv::ComputePhysicalMeshNormals(FClothingActorNv &Actor)
 
 	const uint32 NumParticles = Actor.LodData[CurrentClothingLod].Cloth->getNumParticles();
 	nv::cloth::MappedRange<physx::PxVec4> Particles = Actor.LodData[CurrentClothingLod].Cloth->getCurrentParticles();
-	const TArray<uint32>& Indices = Actor.AssetCreatedFrom->ClothLodData[CurrentClothingLod]->PhysicalMeshData->Indices;
+	const TArray<uint32>& Indices = Actor.AssetCreatedFrom->ClothLodData[CurrentClothingLod]->ClothPhysicalMeshData.Indices;
 	const uint32 NumIndices = Indices.Num();
 
 	// Using the face normals, calculate normals. These will not be normalized as we're adding together
@@ -1044,7 +1038,7 @@ void FClothingSimulationNv::GatherStats() const
 			continue;
 		}
 
-		INC_DWORD_STAT_BY(STAT_NumClothVerts, Actor.AssetCreatedFrom->ClothLodData[Actor.CurrentLodIndex]->PhysicalMeshData->Vertices.Num());
+		INC_DWORD_STAT_BY(STAT_NumClothVerts, Actor.AssetCreatedFrom->ClothLodData[Actor.CurrentLodIndex]->ClothPhysicalMeshData.Vertices.Num());
 	}
 }
 
@@ -1280,7 +1274,7 @@ void FClothingSimulationNv::UpdateLod(int32 InPredictedLod, const FTransform& Co
 					if(CSTransforms.Num() > Actor.AssetCreatedFrom->ReferenceBoneIndex)
 					{
 						// compute skinned positions to init sim mesh
-						const UClothPhysicalMeshDataBase& PhysMesh = *Actor.AssetCreatedFrom->ClothLodData[PredictedClothingLod]->PhysicalMeshData;
+						const FClothPhysicalMeshData& PhysMesh = Actor.AssetCreatedFrom->ClothLodData[PredictedClothingLod]->ClothPhysicalMeshData;
 						TArray<FVector> SkinnedPhysicsMeshPositions;
 						TArray<FVector> SkinnedPhysicsMeshNormals;
 						ClothingMeshUtils::SkinPhysicsMesh(
@@ -1352,10 +1346,8 @@ void FClothingSimulationNv::DebugDraw_PhysMesh(USkeletalMeshComponent* OwnerComp
 		const int32 CurrentClothLod = Actor.CurrentLodIndex;
 
 		nv::cloth::Cloth* CurrentCloth = Actor.LodData[CurrentClothLod].Cloth;
-		
 		check(CurrentCloth);
 
-
 		FTransform RootBoneTransform = OwnerComponent->GetComponentSpaceTransforms()[Actor.AssetCreatedFrom->ReferenceBoneIndex];
 		RootBoneTransform.SetScale3D(FVector(1.0f));
 
@@ -1370,14 +1362,9 @@ void FClothingSimulationNv::DebugDraw_PhysMesh(USkeletalMeshComponent* OwnerComp
 			Particles = NvClothSupport::CreateRange(ActorData.Px_RestPositions);
 		}
 
-		UClothPhysicalMeshDataBase* PhysMeshPtr = Actor.AssetCreatedFrom->ClothLodData[CurrentClothLod]->PhysicalMeshData;
-		TArray<uint32>& Indices = PhysMeshPtr->Indices;
-		UClothPhysicalMeshDataNv* PhysMeshNvPtr = Cast<UClothPhysicalMeshDataNv>(PhysMeshPtr);
-		if (!ensure(PhysMeshNvPtr))
-		{
-			continue;
-		}
-		TArray<float>& MaxDistances = PhysMeshNvPtr->MaxDistances;
+		const FClothPhysicalMeshData& PhysMesh = Actor.AssetCreatedFrom->ClothLodData[CurrentClothLod]->ClothPhysicalMeshData;
+		const TArray<uint32>& Indices = PhysMesh.Indices;
+		const FPointWeightMap& MaxDistances = PhysMesh.GetWeightMap(EWeightMapTargetCommon::MaxDistance);
 
 		const int32 NumTriangles = Indices.Num() / 3;
 
@@ -1528,39 +1515,32 @@ void FClothingSimulationNv::DebugDraw_Backstops(USkeletalMeshComponent* OwnerCom
 		{
 			continue;
 		}
-		const UClothPhysicalMeshDataBase* MeshDataBase = Asset->ClothLodData[Actor.CurrentLodIndex]->PhysicalMeshData;
-		if (!MeshDataBase)
+		const FClothPhysicalMeshData& PhysMesh = Asset->ClothLodData[Actor.CurrentLodIndex]->ClothPhysicalMeshData;
+		const FPointWeightMap* const BackstopDistances = PhysMesh.FindWeightMap(EWeightMapTargetCommon::BackstopDistance);
+		const FPointWeightMap* const BackstopRadiuses = PhysMesh.FindWeightMap(EWeightMapTargetCommon::BackstopRadius);
+		const FPointWeightMap* const MaxDistances = PhysMesh.FindWeightMap(EWeightMapTargetCommon::MaxDistance);
+		if (!BackstopDistances || !BackstopDistances->Num() || 
+			!BackstopRadiuses || !BackstopRadiuses->Num() ||
+			!MaxDistances || !MaxDistances->Num())
 		{
 			continue;
 		}
-		const TArray<float>* BackstopDistancesPtr = MeshDataBase->GetFloatArray((int32)MaskTarget_PhysMesh::BackstopDistance);
-		const TArray<float>* BackstopRadiusesPtr = MeshDataBase->GetFloatArray((int32)MaskTarget_PhysMesh::BackstopRadius);
-		const TArray<float>* MaxDistancesPtr = MeshDataBase->GetFloatArray((int32)MaskTarget_PhysMesh::MaxDistance);
-		if (!BackstopDistancesPtr || !BackstopDistancesPtr->Num() || 
-			!BackstopRadiusesPtr || !BackstopRadiusesPtr->Num() ||
-			!MaxDistancesPtr || !MaxDistancesPtr->Num())
-		{
-			continue;
-		}
-		const TArray<float>& BackstopDistances = *BackstopDistancesPtr;
-		const TArray<float>& BackstopRadiuses = *BackstopRadiusesPtr;
-		const TArray<float>& MaxDistances = *MaxDistancesPtr;
 
 		FTransform RootBoneTransform = OwnerComponent->GetComponentSpaceTransforms()[Actor.AssetCreatedFrom->ReferenceBoneIndex];
 
 		const TArray<FVector>& SkinnedPositions = Actor.GetCurrentSkinnedPositions();
 		const int32 NumVerts = SkinnedPositions.Num();
 		check(NumVerts == Actor.SkinnedPhysicsMeshNormals.Num());
-		check(NumVerts == MaxDistances.Num());
+		check(NumVerts == MaxDistances->Num());
 
 		for(int32 VertIndex = 0; VertIndex < NumVerts; ++VertIndex)
 		{
 			const FVector& Position = RootBoneTransform.TransformPosition(SkinnedPositions[VertIndex]);
 			const FVector& Normal = RootBoneTransform.TransformVector(Actor.SkinnedPhysicsMeshNormals[VertIndex]);
-			
-			float BackstopDistance = BackstopDistances[VertIndex];
-			const float BackstopRadius = BackstopRadiuses[VertIndex];
-			const float MaxDistance = MaxDistances[VertIndex];
+
+			float BackstopDistance = (*BackstopDistances)[VertIndex];
+			const float BackstopRadius = (*BackstopRadiuses)[VertIndex];
+			const float MaxDistance = (*MaxDistances)[VertIndex];
 
 			FColor FixedPointColor = FColor::White;
 
@@ -1609,30 +1589,25 @@ void FClothingSimulationNv::DebugDraw_MaxDistances(USkeletalMeshComponent* Owner
 		{
 			continue;
 		}
-		const UClothPhysicalMeshDataBase* MeshDataBase = Asset->ClothLodData[Actor.CurrentLodIndex]->PhysicalMeshData;
-		if (!MeshDataBase)
-		{
-			continue;
-		}
-		const TArray<float>* MaxDistancesPtr = MeshDataBase ->GetFloatArray((int32)MaskTarget_PhysMesh::MaxDistance);
-		if (!MaxDistancesPtr || !MaxDistancesPtr->Num())
+		const FClothPhysicalMeshData& PhysMesh = Asset->ClothLodData[Actor.CurrentLodIndex]->ClothPhysicalMeshData;
+		const FPointWeightMap* const MaxDistances = PhysMesh.FindWeightMap(EWeightMapTargetCommon::MaxDistance);
+		if (!MaxDistances || !MaxDistances->Num())
 		{
 			continue;
 		}
-		const TArray<float>& MaxDistances = *MaxDistancesPtr;
 
 		FTransform RootBoneTransform = OwnerComponent->GetComponentSpaceTransforms()[Actor.AssetCreatedFrom->ReferenceBoneIndex];
 
 		const TArray<FVector>& SkinnedPositions = Actor.GetCurrentSkinnedPositions();
 		const int32 NumVerts = SkinnedPositions.Num();
 		check(NumVerts == Actor.SkinnedPhysicsMeshNormals.Num());
-		check(NumVerts == MaxDistances.Num());
+		check(NumVerts == MaxDistances->Num());
 
 		for(int32 VertIndex = 0; VertIndex < NumVerts; ++VertIndex)
 		{
 			const FVector& Position = RootBoneTransform.TransformPosition(SkinnedPositions[VertIndex]);
 			const FVector& Normal = RootBoneTransform.TransformVector(Actor.SkinnedPhysicsMeshNormals[VertIndex]);
-			const float& MaxDistance = MaxDistances[VertIndex];
+			const float& MaxDistance = (*MaxDistances)[VertIndex];
 
 			PDI->DrawLine(Position, Position + Normal * MaxDistance, FColor::White, SDPG_World, 0.0f, 0.001f);
 		}
@@ -1649,14 +1624,9 @@ void FClothingSimulationNv::DebugDraw_SelfCollision(USkeletalMeshComponent* Owne
 			continue;
 		}
 
-		const UClothingAssetCommon* Asset = Actor.AssetCreatedFrom;
-		UClothConfigBase* ClothConfigBase = Actor.AssetCreatedFrom->ClothSimConfig;
-		check(ClothConfigBase);
-		UClothConfigNv* ConfigPtr = Cast<UClothConfigNv>(ClothConfigBase);
-		check(ConfigPtr);
-		UClothConfigNv& Config = *ConfigPtr;
-
-		if(!Config.HasSelfCollision())
+		const UClothingAssetCommon* const Asset = Actor.AssetCreatedFrom;
+		const UClothConfigNv* const Config = Asset->GetClothConfig<UClothConfigNv>();
+		if (!Config || !Config->UseSelfCollisions())
 		{
 			// No self collisions on this actor
 			continue;
@@ -1664,13 +1634,12 @@ void FClothingSimulationNv::DebugDraw_SelfCollision(USkeletalMeshComponent* Owne
 
 		FTransform RootBoneTransform = OwnerComponent->GetComponentSpaceTransforms()[Actor.AssetCreatedFrom->ReferenceBoneIndex];
 
-		const float SelfCollisionThickness = Config.SelfCollisionRadius;
+		const float SelfCollisionThickness = Config->SelfCollisionRadius;
 
 		const UClothLODDataCommon& LodData = *Asset->ClothLodData[Actor.CurrentLodIndex];
-		const UClothPhysicalMeshDataBase& PhysMesh = *LodData.PhysicalMeshData;
+		const FClothPhysicalMeshData& PhysMesh = LodData.ClothPhysicalMeshData;
 
 		nv::cloth::Cloth* CurrentCloth = Actor.LodData[Actor.CurrentLodIndex].Cloth;
-
 		check(CurrentCloth);
 
 		NvClothSupport::ClothParticleScopeLock ParticleLoc(CurrentCloth);
@@ -1697,7 +1666,7 @@ void FClothingSimulationNv::DebugDraw_AnimDrive(USkeletalMeshComponent* OwnerCom
 		}
 
 		const UClothingAssetCommon* Asset = Actor.AssetCreatedFrom;
-		const UClothPhysicalMeshDataBase& PhysMesh = *Asset->ClothLodData[Actor.CurrentLodIndex]->PhysicalMeshData;
+		const FClothPhysicalMeshData& PhysMesh = Asset->ClothLodData[Actor.CurrentLodIndex]->ClothPhysicalMeshData;
 		const TArray<uint32>& Indices = PhysMesh.Indices;
 		FTransform RootBoneTransform = OwnerComponent->GetComponentSpaceTransforms()[Actor.AssetCreatedFrom->ReferenceBoneIndex];
 
@@ -1730,7 +1699,10 @@ void FClothingSimulationNv::RefreshClothConfig()
 
 	for(FClothingActorNv& Actor : Actors)
 	{
-		ApplyClothConfig(Actor.AssetCreatedFrom->ClothSimConfig, Actor);
+		if (const UClothConfigNv* const Config = Actor.AssetCreatedFrom->GetClothConfig<UClothConfigNv>())
+		{
+			ApplyClothConfig(Config, Actor);
+		}
 	}
 }
 
@@ -1754,7 +1726,7 @@ FClothingActorNv::FClothingActorNv()
 	, CurrentLodIndex(INDEX_NONE)
 	, bCollisionsDirty(true)
 	, SimDataIndex(INDEX_NONE)
-	, WindMethod(EClothingWindMethod::Legacy)
+	, WindMethod(EClothingWindMethodNv::Legacy)
 	, CurrentSkinnedPositionIndex(0)
 	, PreviousTimestep(0.0f)
 {
@@ -1763,7 +1735,7 @@ FClothingActorNv::FClothingActorNv()
 
 void FClothingActorNv::SkinPhysicsMesh(FClothingSimulationContextNv* InContext)
 {
-	const UClothPhysicalMeshDataBase& PhysMesh = *AssetCreatedFrom->ClothLodData[CurrentLodIndex]->PhysicalMeshData;
+	const FClothPhysicalMeshData& PhysMesh = AssetCreatedFrom->ClothLodData[CurrentLodIndex]->ClothPhysicalMeshData;
 	FTransform RootBoneTransform = InContext->BoneTransforms[AssetCreatedFrom->ReferenceBoneIndex];
 	ClothingMeshUtils::SkinPhysicsMesh(AssetCreatedFrom->UsedBoneIndices, PhysMesh, RootBoneTransform, InContext->RefToLocals.GetData(), InContext->RefToLocals.Num(), SkinnedPhysicsMeshPositions[CurrentSkinnedPositionIndex], SkinnedPhysicsMeshNormals);
 }
@@ -1775,39 +1747,38 @@ void FClothingActorNv::UpdateMotionConstraints(FClothingSimulationContextNv* InC
 		return;
 	}
 
-	nv::cloth::Cloth* CurrentCloth = LodData[CurrentLodIndex].Cloth;
-
+	const TArray<FVector>& CurrentSkinnedPositions = GetCurrentSkinnedPositions();
+	nv::cloth::Cloth* const CurrentCloth = LodData[CurrentLodIndex].Cloth;
 	check(CurrentCloth);
 
-	UClothPhysicalMeshDataNv* PhysMeshPtr = Cast<UClothPhysicalMeshDataNv>(AssetCreatedFrom->ClothLodData[CurrentLodIndex]->PhysicalMeshData);
-	if (!ensure(PhysMeshPtr))
+	const FClothPhysicalMeshData& PhysMesh = AssetCreatedFrom->ClothLodData[CurrentLodIndex]->ClothPhysicalMeshData;
+
+	const FPointWeightMap* const MaxDistances = PhysMesh.FindWeightMap(EWeightMapTargetCommon::MaxDistance);
+	if (MaxDistances && MaxDistances->Num())
 	{
-		return;
-	}
-	TArray<float>& MaxDistances = PhysMeshPtr->MaxDistances;
-	TArray<float>& BackstopDistances = PhysMeshPtr->BackstopDistances;
-	TArray<float>& BackstopRadiuses = PhysMeshPtr->BackstopRadiuses;
-	const TArray<FVector>& CurrentSkinnedPositions = GetCurrentSkinnedPositions();
 
-	nv::cloth::Range<physx::PxVec4> MotionConstraints = CurrentCloth->getMotionConstraints();
-	const int32 NumMotionConstraints = CurrentCloth->getNumMotionConstraints();
-	check(NumMotionConstraints <= CurrentSkinnedPositions.Num());
+		nv::cloth::Range<physx::PxVec4> MotionConstraints = CurrentCloth->getMotionConstraints();
+		const int32 NumMotionConstraints = CurrentCloth->getNumMotionConstraints();
+		check(NumMotionConstraints <= CurrentSkinnedPositions.Num());
 
-	for(int32 ConstraintIndex = 0; ConstraintIndex < NumMotionConstraints; ++ConstraintIndex)
-	{
-		MotionConstraints[ConstraintIndex] = physx::PxVec4(U2PVector(CurrentSkinnedPositions[ConstraintIndex]), MaxDistances[ConstraintIndex] * InContext->MaxDistanceScale);
+		for(int32 ConstraintIndex = 0; ConstraintIndex < NumMotionConstraints; ++ConstraintIndex)
+		{
+			MotionConstraints[ConstraintIndex] = physx::PxVec4(U2PVector(CurrentSkinnedPositions[ConstraintIndex]), (*MaxDistances)[ConstraintIndex] * InContext->MaxDistanceScale);
+		}
 	}
 
-	if(BackstopDistances.Num() && BackstopRadiuses.Num())
+	const FPointWeightMap* const BackstopDistances = PhysMesh.FindWeightMap(EWeightMapTargetCommon::BackstopDistance);
+	const FPointWeightMap* const BackstopRadiuses = PhysMesh.FindWeightMap(EWeightMapTargetCommon::BackstopRadius);
+	if(BackstopDistances && BackstopDistances->Num() && BackstopRadiuses && BackstopRadiuses->Num())
 	{
 		nv::cloth::Range<physx::PxVec4> SeparationConstraints = CurrentCloth->getSeparationConstraints();
 		const int32 NumSeparationConstraints = CurrentCloth->getNumSeparationConstraints();
 		check(NumSeparationConstraints <= SkinnedPhysicsMeshNormals.Num());
-		check(BackstopDistances.Num() == BackstopRadiuses.Num());
+		check(BackstopDistances->Num() == BackstopRadiuses->Num());
 
 		for(int32 ConstraintIndex = 0; ConstraintIndex < NumSeparationConstraints; ++ConstraintIndex)
 		{
-			SeparationConstraints[ConstraintIndex] = physx::PxVec4(U2PVector(CurrentSkinnedPositions[ConstraintIndex] - BackstopDistances[ConstraintIndex] * SkinnedPhysicsMeshNormals[ConstraintIndex]), BackstopRadiuses[ConstraintIndex]);
+			SeparationConstraints[ConstraintIndex] = physx::PxVec4(U2PVector(CurrentSkinnedPositions[ConstraintIndex] - (*BackstopDistances)[ConstraintIndex] * SkinnedPhysicsMeshNormals[ConstraintIndex]), (*BackstopRadiuses)[ConstraintIndex]);
 		}
 	}
 }
@@ -1819,7 +1790,7 @@ void FClothingActorNv::UpdateWind(FClothingSimulationContextNv* InContext, const
 		default:
 			break;
 
-		case EClothingWindMethod::Legacy:
+		case EClothingWindMethodNv::Legacy:
 		{
 			const FTransform RootBoneWorldTransform = InContext->BoneTransforms[AssetCreatedFrom->ReferenceBoneIndex] * InContext->ComponentToWorld;
 			const FVector TransformedWindVelocity = RootBoneWorldTransform.InverseTransformVector(InWindVelocity);
@@ -1827,13 +1798,12 @@ void FClothingActorNv::UpdateWind(FClothingSimulationContextNv* InContext, const
 			TArray<FVector>& ParticleVelocities = Scratch.ParticleVelocities;
 			CalculateParticleVelocities(ParticleVelocities);
 
-			const UClothPhysicalMeshDataBase* MeshDataBase = AssetCreatedFrom->ClothLodData[CurrentLodIndex]->PhysicalMeshData;
-			const TArray<float>* MaxDistancesPtr = MeshDataBase ? MeshDataBase->GetFloatArray((int32)MaskTarget_PhysMesh::MaxDistance) : nullptr;
-			if (!MaxDistancesPtr || !MaxDistancesPtr->Num())
+			const FClothPhysicalMeshData& PhysMesh = AssetCreatedFrom->ClothLodData[CurrentLodIndex]->ClothPhysicalMeshData;
+			const FPointWeightMap* const MaxDistances = PhysMesh.FindWeightMap(EWeightMapTargetCommon::MaxDistance);
+			if (!MaxDistances || !MaxDistances->Num())
 			{
 				return;
 			}
-			const TArray<float>& MaxDistances = *MaxDistancesPtr;
 
 			const int32 NumAccelerations = LodData[CurrentLodIndex].Cloth->getNumParticleAccelerations();
 			nv::cloth::Range<physx::PxVec4> ParticleAccelerations = LodData[CurrentLodIndex].Cloth->getParticleAccelerations();
@@ -1843,7 +1813,7 @@ void FClothingActorNv::UpdateWind(FClothingSimulationContextNv* InContext, const
 				const FVector& Velocity = ParticleVelocities[AccelerationIndex];
 				FVector VelocityDelta = (TransformedWindVelocity * 2500.0f - Velocity);
 
-				if(MaxDistances[AccelerationIndex] > 0.0f && !VelocityDelta.IsZero())
+				if((*MaxDistances)[AccelerationIndex] > 0.0f && !VelocityDelta.IsZero())
 				{
 					// scaled by angle
 					const float DirectionDot = FVector::DotProduct(VelocityDelta.GetUnsafeNormal(), CurrentNormals[AccelerationIndex]);
@@ -1858,7 +1828,7 @@ void FClothingActorNv::UpdateWind(FClothingSimulationContextNv* InContext, const
 		}
 		break;
 
-		case EClothingWindMethod::Accurate:
+		case EClothingWindMethodNv::Accurate:
 		{
 			const physx::PxVec3 PxWindVelocity = U2PVector(InWindVelocity * NvClothSupport::Constants::UnitConversionScale);
 			LodData[CurrentLodIndex].Cloth->setWindVelocity((PxWindVelocity));
@@ -1919,22 +1889,17 @@ void FClothingActorNv::UpdateAnimDrive(FClothingSimulationContextNv* InContext)
 {
 	SCOPE_CYCLE_COUNTER(STAT_NvClothUpdateAnimDrive);
 
-	const UClothPhysicalMeshDataBase* MeshDataBase = AssetCreatedFrom->ClothLodData[CurrentLodIndex]->PhysicalMeshData;
-	if (!MeshDataBase)
-	{
-		return;
-	}
-	const TArray<float>* AnimDriveMultipliers = MeshDataBase->GetFloatArray((int32)MaskTarget_PhysMesh::AnimDriveMultiplier);
+	const FClothPhysicalMeshData& MeshDataBase = AssetCreatedFrom->ClothLodData[CurrentLodIndex]->ClothPhysicalMeshData;
+	const FPointWeightMap* const AnimDriveMultipliers = MeshDataBase.FindWeightMap(EWeightMapTargetCommon::AnimDriveMultiplier);
 	if(AnimDriveMultipliers && AnimDriveMultipliers->Num())
 	{
-		const TArray<float>& AnimDrives = *AnimDriveMultipliers;
 		const TArray<FVector> CurrentSkinPositions = GetCurrentSkinnedPositions();
 		const TArray<FVector> PreviousSkinPositions = GetPreviousSkinnedPositions();
 
 		FActorLodData& CurrLodData = LodData[CurrentLodIndex];
 
 		int32 NumAccelerations = CurrLodData.Cloth->getNumParticleAccelerations();
-		checkSlow(NumAccelerations == 0 || NumAccelerations == AnimDrives.Num());
+		checkSlow(NumAccelerations == 0 || NumAccelerations == AnimDriveMultipliers->Num());
 		nv::cloth::Range<physx::PxVec4> Accelerations = CurrLodData.Cloth->getParticleAccelerations();
 
 		nv::cloth::Range<const physx::PxVec4> Locations = CurrLodData.Cloth->getCurrentParticles();
@@ -1955,7 +1920,7 @@ void FClothingActorNv::UpdateAnimDrive(FClothingSimulationContextNv* InContext)
 			const FVector SkinnedVelocity = CurrentSkinPositions[ParticleIdx] - PreviousSkinPositions[ParticleIdx];
 
 			// accel = (kX - bv)/m (spring with damper)
-			const FVector AddedAcceleration = ((CurrentAnimDriveSpringStiffness * AnimDrives[ParticleIdx] * ToSkinned) - (CurrentAnimDriveDamperStiffness * (ParticleVelocity - SkinnedVelocity))) / (1.0f / Location.w);
+			const FVector AddedAcceleration = ((CurrentAnimDriveSpringStiffness * (*AnimDriveMultipliers)[ParticleIdx] * ToSkinned) - (CurrentAnimDriveDamperStiffness * (ParticleVelocity - SkinnedVelocity))) / (1.0f / Location.w);
 
 			Accelerations[ParticleIdx] += physx::PxVec4(U2PVector(AddedAcceleration), 0.0f);
 		}
diff --git a/Engine/Source/Runtime/ClothingSystemRuntimeNv/Public/Assets/ClothingAssetNv.h b/Engine/Source/Runtime/ClothingSystemRuntimeNv/Public/Assets/ClothingAssetNv.h
deleted file mode 100644
diff --git a/Engine/Source/Runtime/ClothingSystemRuntimeNv/Public/ClothConfigNv.h b/Engine/Source/Runtime/ClothingSystemRuntimeNv/Public/ClothConfigNv.h
--- a/Engine/Source/Runtime/ClothingSystemRuntimeNv/Public/ClothConfigNv.h
+++ b/Engine/Source/Runtime/ClothingSystemRuntimeNv/Public/ClothConfigNv.h
@@ -1,30 +1,20 @@
 // Copyright 1998-2019 Epic Games, Inc. All Rights Reserved.
 #pragma once
 
-#include "Engine/SkeletalMesh.h"
-#include "ClothConfigBase.h"
-#include "ClothConfigNv.generated.h"
+#include "ClothConfig.h"
+#include "ClothConfig_Legacy.h"
 
-namespace nvidia
-{
-	namespace apex
-	{
-		class ClothingAsset;
-	}
-}
+#include "ClothConfigNv.generated.h"
 
-// Container for a constraint setup, these can be horizontal, vertical, shear and bend
+/** Container for a constraint setup, these can be horizontal, vertical, shear and bend. */
 USTRUCT()
-struct FClothConstraintSetup
+struct FClothConstraintSetupNv
 {
 	GENERATED_BODY()
 
-	FClothConstraintSetup()
-		: Stiffness(1.0f)
-		, StiffnessMultiplier(1.0f)
-		, StretchLimit(1.0f)
-		, CompressionLimit(1.0f)
-	{}
+	FClothConstraintSetupNv();
+
+	void MigrateFrom(const FClothConstraintSetup_Legacy& Setup);
 
 	// How stiff this constraint is, this affects how closely it will follow the desired position
 	UPROPERTY(EditAnywhere, Category=Constraint)
@@ -34,7 +24,7 @@ struct FClothConstraintSetup
 	UPROPERTY(EditAnywhere, Category = Constraint)
 	float StiffnessMultiplier;
 
-	// The hard limit on how far this constarint can stretch
+	// The hard limit on how far this constraint can stretch
 	UPROPERTY(EditAnywhere, Category = Constraint)
 	float StretchLimit;
 
@@ -43,8 +33,9 @@ struct FClothConstraintSetup
 	float CompressionLimit;
 };
 
+/** Cloth wind method. */
 UENUM()
-enum class EClothingWindMethod : uint8
+enum class EClothingWindMethodNv : uint8
 {
 	// Use legacy wind mode, where accelerations are modified directly by the simulation
 	// with no regard for drag or lift
@@ -57,41 +48,53 @@ enum class EClothingWindMethod : uint8
 };
 
 /** Holds initial, asset level config for clothing actors. */
+PRAGMA_DISABLE_DEPRECATION_WARNINGS
 UCLASS()
-class CLOTHINGSYSTEMRUNTIMENV_API UClothConfigNv : public UClothConfigBase
+class CLOTHINGSYSTEMRUNTIMENV_API UClothConfigNv : public UClothConfigCommon
 {
 	GENERATED_BODY()
 public:
 	UClothConfigNv();
 
-	bool HasSelfCollision() const override;
+	// Set up custom version serialization.
+	virtual void Serialize(FArchive& Ar) override;
+
+	// Update the deprecated properties.
+	virtual void PostLoad() override;
+
+	// Migrate from the legacy FClothConfig structure.
+	virtual void MigrateFrom(const FClothConfig_Legacy& ClothConfig) override;
+
+	// Return whether self collision is enabled for this config.
+	UE_DEPRECATED(4.25, "This function is deprecated. Please use NeedsSelfCollisionIndices or UseSelfCollisions instead.")
+	virtual bool HasSelfCollision() const override
+	{ return UseSelfCollisions(); }
+
+	// Return the collision radius required to calculate the self collision indices, or 0.f if self collision is disabled.
+	virtual float NeedsSelfCollisionIndices() const override;
 
-#if WITH_EDITOR
-	virtual bool InitFromApexAssetCallback(
-		nvidia::apex::ClothingAsset* InApexAsset,
-		USkeletalMesh* TargetMesh,
-		FName InName);
-#endif // WITH_EDITOR
+	// Return whether this Nv config has self collision.
+	bool UseSelfCollisions() const;
 
 	// How wind should be processed, Accurate uses drag and lift to make the cloth react differently, legacy applies similar forces to all clothing without drag and lift (similar to APEX)
 	UPROPERTY(EditAnywhere, Category = ClothConfig)
-	EClothingWindMethod WindMethod;
+	EClothingWindMethodNv ClothingWindMethod;
 
 	// Constraint data for vertical constraints
 	UPROPERTY(EditAnywhere, Category = ClothConfig)
-	FClothConstraintSetup VerticalConstraintConfig;
+	FClothConstraintSetupNv VerticalConstraint;
 
 	// Constraint data for horizontal constraints
 	UPROPERTY(EditAnywhere, Category = ClothConfig)
-	FClothConstraintSetup HorizontalConstraintConfig;
+	FClothConstraintSetupNv HorizontalConstraint;
 
 	// Constraint data for bend constraints
 	UPROPERTY(EditAnywhere, Category = ClothConfig)
-	FClothConstraintSetup BendConstraintConfig;
+	FClothConstraintSetupNv BendConstraint;
 
 	// Constraint data for shear constraints
 	UPROPERTY(EditAnywhere, Category = ClothConfig)
-	FClothConstraintSetup ShearConstraintConfig;
+	FClothConstraintSetupNv ShearConstraint;
 
 	// Size of self collision spheres centered on each vert
 	UPROPERTY(EditAnywhere, Category = ClothConfig)
@@ -186,178 +189,17 @@ public:
 	// Default damper stiffness for anim drive if an anim drive is in use
 	UPROPERTY(EditAnywhere, Category = ClothConfig)
 	float AnimDriveDamperStiffness;
-};
-
-
-
-/** Deprecated. Use UClothConfigNv instead. */
-USTRUCT()
-struct FClothConfig
-{
-	GENERATED_BODY()
-
-	FClothConfig()
-		: WindMethod(EClothingWindMethod::Legacy)
-		, SelfCollisionRadius(0.0f)
-		, SelfCollisionStiffness(0.0f)
-		, SelfCollisionCullScale(1.0f)
-		, Damping(0.4f)
-		, Friction(0.1f)
-		, WindDragCoefficient(0.02f/100.0f)
-		, WindLiftCoefficient(0.02f/100.0f)
-		, LinearDrag(0.2f)
-		, AngularDrag(0.2f)
-		, LinearInertiaScale(1.0f)
-		, AngularInertiaScale(1.0f)
-		, CentrifugalInertiaScale(1.0f)
-		, SolverFrequency(120.0f)
-		, StiffnessFrequency(100.0f)
-		, GravityScale(1.0f)
-		, GravityOverride(FVector::ZeroVector)
-		, bUseGravityOverride(false)
-		, TetherStiffness(1.0f)
-		, TetherLimit(1.0f)
-		, CollisionThickness(1.0f)
-		, AnimDriveSpringStiffness(1.0f)
-		, AnimDriveDamperStiffness(1.0f)
-	{}
-
-	void
-	MigrateTo(UClothConfigNv* Config)
-	{
-		Config->WindMethod = WindMethod;
-		Config->VerticalConstraintConfig = VerticalConstraintConfig;
-		Config->HorizontalConstraintConfig = HorizontalConstraintConfig;
-		Config->BendConstraintConfig = BendConstraintConfig;
-		Config->ShearConstraintConfig = ShearConstraintConfig;
-		Config->SelfCollisionRadius = SelfCollisionRadius;
-		Config->SelfCollisionStiffness = SelfCollisionStiffness;
-		Config->SelfCollisionCullScale = SelfCollisionCullScale;
-		Config->Damping = Damping;
-		Config->Friction = Friction;
-		Config->WindDragCoefficient = WindDragCoefficient;
-		Config->WindLiftCoefficient = WindLiftCoefficient;
-		Config->LinearDrag = LinearDrag;
-		Config->AngularDrag = AngularDrag;
-		Config->LinearInertiaScale = LinearInertiaScale;
-		Config->AngularInertiaScale = AngularInertiaScale;
-		Config->CentrifugalInertiaScale = CentrifugalInertiaScale;
-		Config->SolverFrequency = SolverFrequency;
-		Config->StiffnessFrequency = StiffnessFrequency;
-		Config->GravityScale = GravityScale;
-		Config->GravityOverride = GravityOverride;
-		Config->bUseGravityOverride = bUseGravityOverride;
-		Config->TetherStiffness = TetherStiffness;
-		Config->TetherLimit = TetherLimit;
-		Config->CollisionThickness = CollisionThickness;
-		Config->AnimDriveSpringStiffness = AnimDriveSpringStiffness;
-		Config->AnimDriveDamperStiffness = AnimDriveDamperStiffness;
-	}
-
-	// How wind should be processed, Accurate uses drag and lift to make the cloth react differently, legacy applies similar forces to all clothing without drag and lift (similar to APEX)
-	UPROPERTY(EditAnywhere, Category = ClothConfig)
-	EClothingWindMethod WindMethod;
-
-	// Constraint data for vertical constraints
-	UPROPERTY(EditAnywhere, Category = ClothConfig)
-	FClothConstraintSetup VerticalConstraintConfig;
-
-	// Constraint data for horizontal constraints
-	UPROPERTY(EditAnywhere, Category = ClothConfig)
-	FClothConstraintSetup HorizontalConstraintConfig;
-
-	// Constraint data for bend constraints
-	UPROPERTY(EditAnywhere, Category = ClothConfig)
-	FClothConstraintSetup BendConstraintConfig;
-
-	// Constraint data for shear constraints
-	UPROPERTY(EditAnywhere, Category = ClothConfig)
-	FClothConstraintSetup ShearConstraintConfig;
-
-	// Size of self collision spheres centered on each vert
-	UPROPERTY(EditAnywhere, Category = ClothConfig)
-	float SelfCollisionRadius;
-
-	// Stiffness of the spring force that will resolve self collisions
-	UPROPERTY(EditAnywhere, Category = ClothConfig)
-	float SelfCollisionStiffness;
-
-	UPROPERTY(EditAnywhere, Category = ClothConfig, meta = (UIMin="0", ClampMin="0"))
-	float SelfCollisionCullScale;
-
-	// Damping of particle motion per-axis
-	UPROPERTY(EditAnywhere, Category = ClothConfig)
-	FVector Damping;
-
-	// Friction of the surface when colliding
-	UPROPERTY(EditAnywhere, Category = ClothConfig)
-	float Friction;
-
-	// Drag coefficient for wind calculations, higher values mean wind has more lateral effect on cloth
-	UPROPERTY(EditAnywhere, Category = ClothConfig)
-	float WindDragCoefficient;
-
-	// Lift coefficient for wind calculations, higher values make cloth rise easier in wind
-	UPROPERTY(EditAnywhere, Category = ClothConfig)
-	float WindLiftCoefficient;
 
-	// Drag applied to linear particle movement per-axis
-	UPROPERTY(EditAnywhere, Category = ClothConfig)
-	FVector LinearDrag;
-
-	// Drag applied to angular particle movement, higher values should limit material bending (per-axis)
-	UPROPERTY(EditAnywhere, Category = ClothConfig)
-	FVector AngularDrag;
-
-	// Scale for linear particle inertia, how much movement should translate to linear motion (per-axis)
-	UPROPERTY(EditAnywhere, Category = ClothConfig, meta = (UIMin="0", UIMax="1", ClampMin="0", ClampMax="1"))
-	FVector LinearInertiaScale;
-
-	// Scale for angular particle inertia, how much movement should translate to angular motion (per-axis)
-	UPROPERTY(EditAnywhere, Category = ClothConfig, meta = (UIMin = "0", UIMax = "1", ClampMin = "0", ClampMax = "1"))
-	FVector AngularInertiaScale;
-
-	// Scale for centrifugal particle inertia, how much movement should translate to angular motion (per-axis)
-	UPROPERTY(EditAnywhere, Category = ClothConfig, meta = (UIMin = "0", UIMax = "1", ClampMin = "0", ClampMax = "1"))
-	FVector CentrifugalInertiaScale;
-
-	// Frequency of the position solver, lower values will lead to stretchier, bouncier cloth
-	UPROPERTY(EditAnywhere, Category = ClothConfig, meta = (UIMin = "30", UIMax = "240", ClampMin = "30", ClampMax = "1000"))
-	float SolverFrequency;
-
-	// Frequency for stiffness calculations, lower values will degrade stiffness of constraints
-	UPROPERTY(EditAnywhere, Category = ClothConfig)
-	float StiffnessFrequency;
-
-	// Scale of gravity effect on particles
-	UPROPERTY(EditAnywhere, Category = ClothConfig, meta = (EditCondition = "!bUseGravityOverride"))
-	float GravityScale;
-
-	// Direct gravity override value
-	UPROPERTY(EditAnywhere, Category = ClothConfig, meta = (EditCondition = "bUseGravityOverride"))
-	FVector GravityOverride;
-
-	// Use gravity override value vs gravity scale 
-	UPROPERTY(EditAnywhere, Category = ClothConfig, meta = (InlineEditConditionToggle))
-	bool bUseGravityOverride;
-
-	// Scale for stiffness of particle tethers between each other
-	UPROPERTY(EditAnywhere, Category = ClothConfig)
-	float TetherStiffness;
-
-	// Scale for the limit of particle tethers (how far they can separate)
-	UPROPERTY(EditAnywhere, Category = ClothConfig)
-	float TetherLimit;
-
-	// 'Thickness' of the simulated cloth, used to adjust collisions
-	UPROPERTY(EditAnywhere, Category = ClothConfig)
-	float CollisionThickness;
-
-	// Default spring stiffness for anim drive if an anim drive is in use
-	UPROPERTY(EditAnywhere, Category = ClothConfig)
-	float AnimDriveSpringStiffness;
-
-	// Default damper stiffness for anim drive if an anim drive is in use
-	UPROPERTY(EditAnywhere, Category = ClothConfig)
-	float AnimDriveDamperStiffness;
+	// Deprecated properties using old legacy structure and enum that couldn't be redirected
+	UPROPERTY()
+	EClothingWindMethod_Legacy WindMethod_DEPRECATED;
+	UPROPERTY()
+	FClothConstraintSetup_Legacy VerticalConstraintConfig_DEPRECATED;
+	UPROPERTY()
+	FClothConstraintSetup_Legacy HorizontalConstraintConfig_DEPRECATED;
+	UPROPERTY()
+	FClothConstraintSetup_Legacy BendConstraintConfig_DEPRECATED;
+	UPROPERTY()
+	FClothConstraintSetup_Legacy ShearConstraintConfig_DEPRECATED;
 };
+PRAGMA_ENABLE_DEPRECATION_WARNINGS
diff --git a/Engine/Source/Runtime/ClothingSystemRuntimeNv/Public/ClothLODDataNv.h b/Engine/Source/Runtime/ClothingSystemRuntimeNv/Public/ClothLODDataNv.h
deleted file mode 100644
diff --git a/Engine/Source/Runtime/ClothingSystemRuntimeNv/Public/ClothParameterMask_PhysMesh.h b/Engine/Source/Runtime/ClothingSystemRuntimeNv/Public/ClothParameterMask_PhysMesh.h
deleted file mode 100644
diff --git a/Engine/Source/Runtime/ClothingSystemRuntimeNv/Public/ClothPhysicalMeshDataNv.h b/Engine/Source/Runtime/ClothingSystemRuntimeNv/Public/ClothPhysicalMeshDataNv.h
deleted file mode 100644
diff --git a/Engine/Source/Runtime/ClothingSystemRuntimeNv/Public/ClothPhysicalMeshDataNv_Legacy.h b/Engine/Source/Runtime/ClothingSystemRuntimeNv/Public/ClothPhysicalMeshDataNv_Legacy.h
new file mode 100644
--- /dev/null
+++ b/Engine/Source/Runtime/ClothingSystemRuntimeNv/Public/ClothPhysicalMeshDataNv_Legacy.h
@@ -0,0 +1,34 @@
+// Copyright 1998-2019 Epic Games, Inc. All Rights Reserved.
+#pragma once
+
+#include "ClothPhysicalMeshDataBase_Legacy.h"
+#include "ClothPhysicalMeshDataNv_Legacy.generated.h"
+
+/**
+ * Deprecated, use FClothPhysicalMeshData instead.
+ * NV specific spatial simulation data for a mesh.
+ */
+UCLASS()
+class CLOTHINGSYSTEMRUNTIMENV_API UClothPhysicalMeshDataNv_Legacy: public UClothPhysicalMeshDataBase_Legacy
+{
+	GENERATED_BODY()
+public:
+	UClothPhysicalMeshDataNv_Legacy();
+	virtual ~UClothPhysicalMeshDataNv_Legacy();
+
+	// The distance that each vertex can move away from its reference (skinned) position
+	UPROPERTY()
+	TArray<float> MaxDistances;
+
+	// Distance along the plane of the surface that the particles can travel (separation constraint)
+	UPROPERTY()
+	TArray<float> BackstopDistances;
+
+	// Radius of movement to allow for backstop movement
+	UPROPERTY()
+	TArray<float> BackstopRadiuses;
+
+	// Strength of anim drive per-particle (spring driving particle back to skinned location
+	UPROPERTY()
+	TArray<float> AnimDriveMultipliers;
+};
diff --git a/Engine/Source/Runtime/ClothingSystemRuntimeNv/Public/ClothingSimulationFactoryNv.h b/Engine/Source/Runtime/ClothingSystemRuntimeNv/Public/ClothingSimulationFactoryNv.h
--- a/Engine/Source/Runtime/ClothingSystemRuntimeNv/Public/ClothingSimulationFactoryNv.h
+++ b/Engine/Source/Runtime/ClothingSystemRuntimeNv/Public/ClothingSimulationFactoryNv.h
@@ -2,7 +2,6 @@
 
 #pragma once
 
-#include "ClothingSimulation.h"
 #include "ClothingSimulationFactory.h"
 
 #include "ClothingSimulationFactoryNv.generated.h"
@@ -19,4 +18,7 @@ public:
 
 	virtual bool SupportsRuntimeInteraction() override;
 	virtual UClothingSimulationInteractor* CreateInteractor() override;
-};
\ No newline at end of file
+
+	virtual TSubclassOf<UClothConfigBase> GetClothConfigClass() const override;
+	virtual const UEnum* GetWeightMapTargetEnum() const override;
+};
diff --git a/Engine/Source/Runtime/ClothingSystemRuntimeNv/Public/ClothingSimulationNv.h b/Engine/Source/Runtime/ClothingSystemRuntimeNv/Public/ClothingSimulationNv.h
--- a/Engine/Source/Runtime/ClothingSystemRuntimeNv/Public/ClothingSimulationNv.h
+++ b/Engine/Source/Runtime/ClothingSystemRuntimeNv/Public/ClothingSimulationNv.h
@@ -7,9 +7,10 @@
 #include "ClothingSimulation.h"
 #include "ClothingSimulationContext.h"
 #include "ClothingActor.h"
+#include "ClothConfigNv.h"
 
-#include "Assets/ClothingAssetNv.h"
-
+class UClothingAssetBase;
+class UClothingAssetCommon;
 
 namespace nv
 {
@@ -132,7 +133,7 @@ private:
 	TArray<FActorLodData> LodData;
 
 	// How we're going to calculate our wind data (see EClothingWindMethod for method descriptions)
-	EClothingWindMethod WindMethod;
+	EClothingWindMethodNv WindMethod;
 
 	// Thickness to add to collisions to fake cloth thickness
 	float CollisionThickness;
@@ -157,7 +158,7 @@ private:
 	// Scratch arrays for processing during simulate, grow-only to avoid repeated allocations.
 	FClothingActorScratchData Scratch;
 
-	// Simuation given access to our data
+	// Simulation given access to our data
 	friend class FClothingSimulationNv;
 };
 
@@ -250,15 +251,15 @@ private:
 
 	// The core simulation is only solving unoriented particles, so we need to compute normals after the
 	// simulation runs
-	void ComputePhysicalMeshNormals(FClothingActorNv &Actor);
+	void ComputePhysicalMeshNormals(FClothingActorNv& Actor);
 
 	// Given a clothing config from an asset, apply it to the provided actor. Currently
 	// this is only used from CreateActor, but could be exposed for runtime changes
-	void ApplyClothConfig(UClothConfigBase *BaseConfig, FClothingActorNv &InActor);
+	void ApplyClothConfig(const UClothConfigNv* Config, FClothingActorNv& InActor);
 
 	// Extract collisions from the physics asset inside Asset and apply them to InActor
 	// Not safe to call from workers (i.e. inside the simulation).
-	void ExtractActorCollisions(UClothingAssetCommon* Asset, FClothingActorNv &InActor);
+	void ExtractActorCollisions(UClothingAssetCommon* Asset, FClothingActorNv& InActor);
 
 	// The current LOD index for the owning skeletal mesh component
 	int32 CurrentMeshLodIndex;
