Chaos Cloth - Refactor the cloth simulation and context classes.
- Move the IClothingSimulation and IClothingSimulationContext definitions from ClothingSimulation.h/cpp to ClothingSimulationInterface.h/cpp files to make place for the FClothingSimulationCommon definitions.
- Move common simulation and context implementation from the ClothingSystemRuntimeInterface to the ClothingSystemRuntimeCommon module.
- Merge common code used in Nv and Chaos cloth simulations, and remove duplicated code where possible.
- Move specific FillContext behaviors out of the base simulation class and into the context subclasses.

Merge: CL 10878385, CL10879173
Release-11.40-Chaos -> Release-12.00

diff --git a/Engine/Plugins/Experimental/ChaosCloth/Source/Chaos/Private/ChaosCloth/ChaosClothingSimulation.cpp b/Engine/Plugins/Experimental/ChaosCloth/Source/Chaos/Private/ChaosCloth/ChaosClothingSimulation.cpp
--- a/Engine/Plugins/Experimental/ChaosCloth/Source/Chaos/Private/ChaosCloth/ChaosClothingSimulation.cpp
+++ b/Engine/Plugins/Experimental/ChaosCloth/Source/Chaos/Private/ChaosCloth/ChaosClothingSimulation.cpp
@@ -1107,76 +1107,14 @@ void ClothingSimulation::AddCollisions(const FClothCollisionData& ClothCollision
 	UE_LOG(LogChaosCloth, VeryVerbose, TEXT("Added collisions: %d spheres, %d capsules, %d convexes, %d boxes."), NumSpheres, NumCapsules, NumConvexes, NumBoxes);
 }
 
-void ClothingSimulation::FillContext(USkeletalMeshComponent* InComponent, float InDeltaTime, IClothingSimulationContext* InOutContext)
-{
-    ClothingSimulationContext* Context = static_cast<ClothingSimulationContext*>(InOutContext);
-    Context->ComponentToWorld = InComponent->GetComponentToWorld();
-    Context->DeltaTime = ClampDeltaTime > 0 ? FMath::Min(InDeltaTime, ClampDeltaTime) : InDeltaTime;
-
-	Context->WorldGravity = InComponent->GetWorld() ?
-		FVector(0.f, 0.f, InComponent->GetWorld()->GetGravityZ()) :
-		ClothingSimulationDefault::Gravity;
-
-	Context->RefToLocals.Reset();
-    InComponent->GetCurrentRefToLocalMatrices(Context->RefToLocals, 0);
-
-	const USkeletalMesh* SkelMesh = InComponent->SkeletalMesh;
-    if (USkinnedMeshComponent* MasterComponent = InComponent->MasterPoseComponent.Get())
-    {
-		const TArray<int32>& MasterBoneMap = InComponent->GetMasterBoneMap();
-        int32 NumBones = MasterBoneMap.Num();
-        if (NumBones == 0)
-        {
-            if (InComponent->SkeletalMesh)
-            {
-                // This case indicates an invalid master pose component (e.g. no skeletal mesh)
-                NumBones = InComponent->SkeletalMesh->RefSkeleton.GetNum();
-            }
-			Context->BoneTransforms.Reset(NumBones);
-			Context->BoneTransforms.AddDefaulted(NumBones);
-        }
-        else
-        {
-            Context->BoneTransforms.Reset(NumBones);
-            Context->BoneTransforms.AddDefaulted(NumBones);
-			const TArray<FTransform>& MasterTransforms = MasterComponent->GetComponentSpaceTransforms();
-            for (int32 BoneIndex = 0; BoneIndex < NumBones; ++BoneIndex)
-            {
-                bool bFoundMaster = false;
-                if (MasterBoneMap.IsValidIndex(BoneIndex))
-                {
-                    const int32 MasterIndex = MasterBoneMap[BoneIndex];
-                    if (MasterTransforms.IsValidIndex(MasterIndex))
-                    {
-                        Context->BoneTransforms[BoneIndex] = MasterTransforms[MasterIndex];
-                        bFoundMaster = true;
-                    }
-                }
-
-                if (!bFoundMaster && SkelMesh)
-                {
-                    const int32 ParentIndex = SkelMesh->RefSkeleton.GetParentIndex(BoneIndex);
-					check(ParentIndex < BoneIndex);
-					Context->BoneTransforms[BoneIndex] =
-						Context->BoneTransforms.IsValidIndex(ParentIndex) && ParentIndex < BoneIndex ?
-						Context->BoneTransforms[ParentIndex] * SkelMesh->RefSkeleton.GetRefBonePose()[BoneIndex] :
-                        SkelMesh->RefSkeleton.GetRefBonePose()[BoneIndex];
-                }
-            }
-        }
-    }
-    else
-    {
-        Context->BoneTransforms = InComponent->GetComponentSpaceTransforms();
-    }
-}
-
 void ClothingSimulation::Simulate(IClothingSimulationContext* InContext)
 {
 	SCOPE_CYCLE_COUNTER(STAT_ChaosClothSimulate);
-	ClothingSimulationContext* Context = static_cast<ClothingSimulationContext*>(InContext);
-	if (Context->DeltaTime == 0)
+	const ClothingSimulationContext* const Context = static_cast<ClothingSimulationContext*>(InContext);
+	if (Context->DeltaSeconds == 0.f)
+	{
 		return;
+	}
 
 	// Set gravity
 	Evolution->GetGravityForces().SetAcceleration(Chaos::TVector<float, 3>(
@@ -1225,7 +1163,7 @@ void ClothingSimulation::Simulate(IClothingSimulationContext* InContext)
 	UpdateCollisionTransforms(*Context);
 
 	// Advance Sim
-	DeltaTime = Context->DeltaTime;
+	DeltaTime = Context->DeltaSeconds;
 
 	const float SolverFrequency = ClothSharedSimConfig ?
 		ClothSharedSimConfig->SolverFrequency :
diff --git a/Engine/Plugins/Experimental/ChaosCloth/Source/Chaos/Public/ChaosCloth/ChaosClothingSimulation.h b/Engine/Plugins/Experimental/ChaosCloth/Source/Chaos/Public/ChaosCloth/ChaosClothingSimulation.h
--- a/Engine/Plugins/Experimental/ChaosCloth/Source/Chaos/Public/ChaosCloth/ChaosClothingSimulation.h
+++ b/Engine/Plugins/Experimental/ChaosCloth/Source/Chaos/Public/ChaosCloth/ChaosClothingSimulation.h
@@ -1,8 +1,7 @@
 // Copyright Epic Games, Inc. All Rights Reserved.
 #pragma once
 
-#include "ClothingSimulationInterface.h"
-
+#include "ClothingSimulation.h"
 #include "ClothingAsset.h"
 #include "Chaos/ArrayCollectionArray.h"
 #include "Chaos/PBDEvolution.h"
@@ -13,20 +12,9 @@
 
 namespace Chaos
 {
-	class ClothingSimulationContext : public IClothingSimulationContext  // TODO(Kriss.Gossart): Inherit from FClothingSimulationContextBase
-	{
-	public:
-		ClothingSimulationContext() {}
-		~ClothingSimulationContext() {}
-
-		float DeltaTime;
-		TArray<FMatrix> RefToLocals;
-		TArray<FTransform> BoneTransforms;
-		FTransform ComponentToWorld;
-		FVector WorldGravity;  // Gravity extracted from the world
-	};
+	typedef FClothingSimulationContextCommon ClothingSimulationContext;
 
-	class ClothingSimulation : public IClothingSimulation
+	class ClothingSimulation : public FClothingSimulationCommon
 #if WITH_EDITOR
 		, public FGCObject  // Add garbage collection for debug cloth material
 #endif  // #if WITH_EDITOR
@@ -70,7 +58,6 @@ namespace Chaos
 		virtual void CreateActor(USkeletalMeshComponent* InOwnerComponent, UClothingAssetBase* InAsset, int32 SimDataIndex) override;
 		virtual void PostActorCreationInitialize() override;
 		virtual IClothingSimulationContext* CreateContext() override { return new ClothingSimulationContext(); }
-		virtual void FillContext(USkeletalMeshComponent* InComponent, float InDeltaTime, IClothingSimulationContext* InOutContext) override;
 		virtual void Shutdown() override;
 		virtual bool ShouldSimulate() const override { return true; }
 		virtual void Simulate(IClothingSimulationContext* InContext) override;
diff --git a/Engine/Source/Runtime/ClothingSystemRuntimeCommon/Private/ClothingSimulation.cpp b/Engine/Source/Runtime/ClothingSystemRuntimeCommon/Private/ClothingSimulation.cpp
new file mode 100644
--- /dev/null
+++ b/Engine/Source/Runtime/ClothingSystemRuntimeCommon/Private/ClothingSimulation.cpp
@@ -0,0 +1,165 @@
+// Copyright Epic Games, Inc. All Rights Reserved.
+
+#include "ClothingSimulation.h"
+#include "Components/SkeletalMeshComponent.h"
+#include "PhysicsEngine/PhysicsSettings.h"
+
+//==============================================================================
+// FClothingSimulationContextCommon
+//==============================================================================
+
+FClothingSimulationContextCommon::FClothingSimulationContextCommon()
+	: ComponentToWorld(FTransform::Identity)
+	, WorldGravity(FVector::ZeroVector)
+	, WindVelocity(FVector::ZeroVector)
+	, DeltaSeconds(0.f)
+{}
+
+FClothingSimulationContextCommon::~FClothingSimulationContextCommon()
+{}
+
+void FClothingSimulationContextCommon::Fill(const USkeletalMeshComponent* InComponent, float InDeltaSeconds, float InMaxPhysicsDelta)
+{
+	check(InComponent);
+	FillBoneTransforms(InComponent);
+	FillRefToLocals(InComponent);
+	FillComponentToWorld(InComponent);
+	FillWorldGravity(InComponent);
+	FillWindVelocity(InComponent);
+	FillDeltaSeconds(InDeltaSeconds, InMaxPhysicsDelta);
+}
+
+void FClothingSimulationContextCommon::FillBoneTransforms(const USkeletalMeshComponent* InComponent)
+{
+	const USkeletalMesh* const SkeletalMesh = InComponent->SkeletalMesh;
+
+	if (USkinnedMeshComponent* const MasterComponent = InComponent->MasterPoseComponent.Get())
+	{
+		const TArray<int32>& MasterBoneMap = InComponent->GetMasterBoneMap();
+		int32 NumBones = MasterBoneMap.Num();
+
+		if (NumBones == 0)
+		{
+			if (SkeletalMesh)
+			{
+				// This case indicates an invalid master pose component (e.g. no skeletal mesh)
+				NumBones = SkeletalMesh->RefSkeleton.GetNum();
+
+				BoneTransforms.Empty(NumBones);
+				BoneTransforms.AddDefaulted(NumBones);
+			}
+		}
+		else
+		{
+			BoneTransforms.Reset(NumBones);
+			BoneTransforms.AddDefaulted(NumBones);
+
+			const TArray<FTransform>& MasterTransforms = MasterComponent->GetComponentSpaceTransforms();
+			for (int32 BoneIndex = 0; BoneIndex < NumBones; ++BoneIndex)
+			{
+				bool bFoundMaster = false;
+				if (MasterBoneMap.IsValidIndex(BoneIndex))
+				{
+					const int32 MasterIndex = MasterBoneMap[BoneIndex];
+					if (MasterIndex != INDEX_NONE && MasterIndex < MasterTransforms.Num())
+					{
+						BoneTransforms[BoneIndex] = MasterTransforms[MasterIndex];
+						bFoundMaster = true;
+					}
+				}
+
+				if (!bFoundMaster && SkeletalMesh)
+				{
+					const int32 ParentIndex = SkeletalMesh->RefSkeleton.GetParentIndex(BoneIndex);
+
+					BoneTransforms[BoneIndex] =
+						BoneTransforms.IsValidIndex(ParentIndex) && ParentIndex < BoneIndex ?
+						BoneTransforms[ParentIndex] * SkeletalMesh->RefSkeleton.GetRefBonePose()[BoneIndex] :
+						SkeletalMesh->RefSkeleton.GetRefBonePose()[BoneIndex];
+				}
+			}
+		}
+	}
+	else
+	{
+		BoneTransforms = InComponent->GetComponentSpaceTransforms();
+	}
+}
+
+void FClothingSimulationContextCommon::FillRefToLocals(const USkeletalMeshComponent* InComponent)
+{
+	RefToLocals.Reset();
+	InComponent->GetCurrentRefToLocalMatrices(RefToLocals, 0);
+}
+
+void FClothingSimulationContextCommon::FillComponentToWorld(const USkeletalMeshComponent* InComponent)
+{
+	ComponentToWorld = InComponent->GetComponentTransform();
+}
+
+void FClothingSimulationContextCommon::FillWorldGravity(const USkeletalMeshComponent* InComponent)
+{
+	const UWorld* const ComponentWorld = InComponent->GetWorld();
+	check(ComponentWorld);
+	WorldGravity = FVector(0.f, 0.f, ComponentWorld->GetGravityZ());
+}
+
+void FClothingSimulationContextCommon::FillWindVelocity(const USkeletalMeshComponent* InComponent)
+{
+	SetWindFromComponent(InComponent);
+}
+
+void FClothingSimulationContextCommon::FillDeltaSeconds(float InDeltaSeconds, float InMaxPhysicsDelta)
+{
+	DeltaSeconds = FMath::Min(InDeltaSeconds, InMaxPhysicsDelta);
+}
+
+float FClothingSimulationContextCommon::SetWindFromComponent(const USkeletalMeshComponent* Component)
+{
+	float WindAdaption;
+	Component->GetWindForCloth_GameThread(WindVelocity, WindAdaption);
+	return WindAdaption;
+}
+
+//==============================================================================
+// FClothingSimulationCommon
+//==============================================================================
+
+FClothingSimulationCommon::FClothingSimulationCommon()
+{
+	MaxPhysicsDelta = UPhysicsSettings::Get()->MaxPhysicsDeltaTime;
+}
+
+FClothingSimulationCommon::~FClothingSimulationCommon()
+{}
+
+void FClothingSimulationCommon::FillContext(USkeletalMeshComponent* InComponent, float InDeltaTime, IClothingSimulationContext* InOutContext)
+{
+	check(InOutContext);
+	FClothingSimulationContextCommon* const Context = static_cast<FClothingSimulationContextCommon*>(InOutContext);
+
+	Context->Fill(InComponent, InDeltaTime, MaxPhysicsDelta);
+
+	// Checking the component here to track rare issue leading to invalid contexts
+	if (InComponent->IsPendingKill())
+	{
+		const AActor* const CompOwner = InComponent->GetOwner();
+		UE_LOG(LogSkeletalMesh, Warning, 
+			TEXT("Attempting to fill a clothing simulation context for a PendingKill skeletal mesh component (Comp: %s, Actor: %s). "
+				"Pending kill skeletal mesh components should be unregistered before marked pending kill."), 
+			*InComponent->GetName(), CompOwner ? *CompOwner->GetName() : TEXT("None"));
+
+		// Make sure we clear this out to skip any attempted simulations
+		Context->BoneTransforms.Reset();
+	}
+
+	if (Context->BoneTransforms.Num() == 0)
+	{
+		const AActor* const CompOwner = InComponent->GetOwner();
+		const USkinnedMeshComponent* const Master = InComponent->MasterPoseComponent.Get();
+		UE_LOG(LogSkeletalMesh, Warning, TEXT("Attempting to fill a clothing simulation context for a skeletal mesh component that has zero bones (Comp: %s, Master: %s, Actor: %s)."), *InComponent->GetName(), Master ? *Master->GetName() : TEXT("None"), CompOwner ? *CompOwner->GetName() : TEXT("None"));
+
+		// Make sure we clear this out to skip any attempted simulations
+		Context->BoneTransforms.Reset();
+	}
+}
diff --git a/Engine/Source/Runtime/ClothingSystemRuntimeCommon/Public/ClothingSimulation.h b/Engine/Source/Runtime/ClothingSystemRuntimeCommon/Public/ClothingSimulation.h
new file mode 100644
--- /dev/null
+++ b/Engine/Source/Runtime/ClothingSystemRuntimeCommon/Public/ClothingSimulation.h
@@ -0,0 +1,67 @@
+// Copyright Epic Games, Inc. All Rights Reserved.
+#pragma once
+
+#include "ClothingSimulationInterface.h"
+#include "Containers/Array.h"
+#include "Math/Transform.h"
+#include "Math/Vector.h"
+
+class USkeletalMeshComponent;
+
+/** Base simulation data that just about every simulation would need. */
+class CLOTHINGSYSTEMRUNTIMECOMMON_API FClothingSimulationContextCommon : public IClothingSimulationContext
+{
+public:
+	FClothingSimulationContextCommon();
+	virtual ~FClothingSimulationContextCommon() override;
+
+	// Fill this context using the given skeletal mesh component
+	virtual void Fill(const USkeletalMeshComponent* InComponent, float InDeltaSeconds, float InMaxPhysicsDelta);
+
+protected:
+	// Default fill behavior as expected to be used by every simulation
+	virtual void FillBoneTransforms(const USkeletalMeshComponent* InComponent);
+	virtual void FillRefToLocals(const USkeletalMeshComponent* InComponent);
+	virtual void FillComponentToWorld(const USkeletalMeshComponent* InComponent);
+	virtual void FillWorldGravity(const USkeletalMeshComponent* InComponent);
+	virtual void FillWindVelocity(const USkeletalMeshComponent* InComponent);
+	virtual void FillDeltaSeconds(float InDeltaSeconds, float InMaxPhysicsDelta);
+
+	// Set the wind velocity and return the wind adaptation if required
+	float SetWindFromComponent(const USkeletalMeshComponent* Component);
+
+public:
+	// World space bone transforms of the owning component
+	TArray<FTransform> BoneTransforms;
+
+	// Ref to local matrices from the owning component (for skinning fixed verts)
+	TArray<FMatrix> RefToLocals;
+
+	// Component to world transform of the owning component
+	FTransform ComponentToWorld;
+
+	// Gravity extracted from the world
+	FVector WorldGravity;
+
+	// Wind velocity at the component location
+	FVector WindVelocity;
+
+	// Delta for this tick
+	float DeltaSeconds;
+};
+
+// Base simulation to fill in common data for the base context
+class CLOTHINGSYSTEMRUNTIMECOMMON_API FClothingSimulationCommon : public IClothingSimulation
+{
+public:
+	FClothingSimulationCommon();
+	virtual ~FClothingSimulationCommon();
+
+protected:
+	/** Fills in the base data for a clothing simulation */
+	virtual void FillContext(USkeletalMeshComponent* InComponent, float InDeltaTime, IClothingSimulationContext* InOutContext) override;
+
+private:
+	/** Maximum physics time, incoming deltas will be clamped down to this value on long frames */
+	float MaxPhysicsDelta;
+};
diff --git a/Engine/Source/Runtime/ClothingSystemRuntimeInterface/Private/ClothingSimulation.cpp b/Engine/Source/Runtime/ClothingSystemRuntimeInterface/Private/ClothingSimulation.cpp
deleted file mode 100644
diff --git a/Engine/Source/Runtime/ClothingSystemRuntimeInterface/Private/ClothingSimulationContext.cpp b/Engine/Source/Runtime/ClothingSystemRuntimeInterface/Private/ClothingSimulationContext.cpp
deleted file mode 100644
diff --git a/Engine/Source/Runtime/ClothingSystemRuntimeInterface/Private/ClothingSimulationInterface.cpp b/Engine/Source/Runtime/ClothingSystemRuntimeInterface/Private/ClothingSimulationInterface.cpp
new file mode 100644
--- /dev/null
+++ b/Engine/Source/Runtime/ClothingSystemRuntimeInterface/Private/ClothingSimulationInterface.cpp
@@ -0,0 +1,23 @@
+// Copyright 1998-2019 Epic Games, Inc. All Rights Reserved.
+
+#include "ClothingSimulationInterface.h"
+
+//==============================================================================
+// IClotingSimulationContext
+//==============================================================================
+
+IClothingSimulationContext::IClothingSimulationContext()
+{}
+
+IClothingSimulationContext::~IClothingSimulationContext()
+{}
+
+//==============================================================================
+// IClothingSimulation
+//==============================================================================
+
+IClothingSimulation::IClothingSimulation()
+{}
+
+IClothingSimulation::~IClothingSimulation()
+{}
diff --git a/Engine/Source/Runtime/ClothingSystemRuntimeInterface/Public/ClothingSimulation.h b/Engine/Source/Runtime/ClothingSystemRuntimeInterface/Public/ClothingSimulation.h
deleted file mode 100644
diff --git a/Engine/Source/Runtime/ClothingSystemRuntimeInterface/Public/ClothingSimulationContext.h b/Engine/Source/Runtime/ClothingSystemRuntimeInterface/Public/ClothingSimulationContext.h
deleted file mode 100644
diff --git a/Engine/Source/Runtime/ClothingSystemRuntimeInterface/Public/ClothingSimulationInterface.h b/Engine/Source/Runtime/ClothingSystemRuntimeInterface/Public/ClothingSimulationInterface.h
--- a/Engine/Source/Runtime/ClothingSystemRuntimeInterface/Public/ClothingSimulationInterface.h
+++ b/Engine/Source/Runtime/ClothingSystemRuntimeInterface/Public/ClothingSimulationInterface.h
@@ -1,5 +1,144 @@
 // Copyright Epic Games, Inc. All Rights Reserved.
 #pragma once
 
-#include "ClothingSimulationContext.h"
-#include "ClothingSimulation.h"
+#include "ClothingSystemRuntimeTypes.h"
+
+class UClothingAssetBase;
+class USkeletalMeshComponent;
+class USkinnedMeshComponent;
+struct FClothCollisionData;
+
+/** Empty interface, derived simulation modules define the contents of the context. */
+class CLOTHINGSYSTEMRUNTIMEINTERFACE_API IClothingSimulationContext
+{
+public:
+	IClothingSimulationContext();
+	virtual ~IClothingSimulationContext();
+};
+
+/** Base class for clothing simulators. */
+class CLOTHINGSYSTEMRUNTIMEINTERFACE_API IClothingSimulation
+{
+public:
+	IClothingSimulation();
+	virtual ~IClothingSimulation();
+
+	// The majority of the API for this class is protected. The required objects (skel meshes and the parallel task)
+	// are friends so they can use the functionality. For the most part the simulation is not designed to be used
+	// outside of the skeletal mesh component as its parallel simulation is tied to the skel mesh tick and
+	// dependents.
+	// Any method that is available in the public section below should consider that it may be called while
+	// the simulation is running.
+	// Currently only a const pointer is exposed externally to the skeletal mesh component, so barring that
+	// being cast away external callers can't call non-const methods. The debug skel mesh component exposes a
+	// mutable version for editor use
+protected:
+
+	friend class USkeletalMeshComponent;
+	friend class FParallelClothTask;
+
+	/**
+	 * Create an actor for this simulation from the data in InAsset
+	 * Simulation data for this actor should be written back to SimDataIndex in GetSimulationData
+	 * @param InOwnerComponent - The component requesting this actor
+	 * @param InAsset - The asset to create an actor from
+	 * @param SimDataIndex - the sim data index to use when doing the writeback in GetSimulationData
+	 */
+	virtual void CreateActor(USkeletalMeshComponent* InOwnerComponent, UClothingAssetBase* InAsset, int32 SimDataIndex) = 0;
+
+	/** Do any required initialization after all the actors have been loaded */
+	virtual void PostActorCreationInitialize() = 0;
+	
+	/** Create a new context, will not be filled, call FillContext before simulating with this context */
+	virtual IClothingSimulationContext* CreateContext() = 0;
+
+	/**
+	 * Fills an existing context for a single simulation step, called by the engine on the game thread prior to simulation 
+	 * @param InComponent - The component to fill the context for
+	 * @param InOutContext - The context to fill
+	 */
+	virtual void FillContext(USkeletalMeshComponent* InComponent, float InDeltaTime, IClothingSimulationContext* InOutContext) = 0;
+
+	/** Initialize the simulation, will be called before any Simulate calls */
+	virtual void Initialize() = 0;
+
+	/**
+	 * Shutdown the simulation, this should clear ALL resources as we no longer expect to
+	 * call any other function on this simulation without first calling Initialize again.
+	 */
+	virtual void Shutdown() = 0;
+
+	/**
+	 * Called by the engine to detect whether or not the simulation should run, essentially
+	 * are there any actors that need to simulate in this simulation
+	 */
+	virtual bool ShouldSimulate() const = 0;
+
+	/**
+	 * Run a single tick of the simulation.
+	 * The pointer InContext is guaranteed (when called by the engine) to be the context allocated in
+	 * CreateContext and can be assumed to be safely castable to any derived type allocated there.
+	 * New callers should make care to make sure only the correct context is ever passed through
+	 * @param InContext - The context to use during simulation, will have been filled in FillContext
+	 */
+	virtual void Simulate(IClothingSimulationContext* InContext) = 0;
+
+	/** Simulation should remove all of it's actors when next possible and free them */
+	virtual void DestroyActors() = 0;
+
+	/** 
+	 * Destroy a context object, engine will always pass a context created using CreateContext 
+	 * @param InContext - The context to destroy
+	 */
+	virtual void DestroyContext(IClothingSimulationContext* InContext) = 0;
+
+	/** 
+	 * Fill FClothSimulData map for the clothing simulation. Should fill a map pair per-actor 
+	 * @param OutData - The simulation data to write to
+	 * @param InOwnerComponent - the component that owns the simulation
+	 * @param InOverrideComponent - An override component if bound to a master pose component
+	 */
+	virtual void GetSimulationData(TMap<int32, FClothSimulData>& OutData, USkeletalMeshComponent* InOwnerComponent, USkinnedMeshComponent* InOverrideComponent) const = 0;
+
+	/** Get the bounds of the simulation mesh in local simulation space */
+	virtual FBoxSphereBounds GetBounds(const USkeletalMeshComponent* InOwnerComponent) const = 0;
+
+	/**
+	 * Called by the engine when an external object wants to inject collision data into this simulation above
+	 * and beyond what is specified in the asset for the internal actors
+	 * Examples: Scene collision, collision for parents we are attached to
+	 * @param InData - Collisions to add to this simulation
+	 */
+	virtual void AddExternalCollisions(const FClothCollisionData& InData) = 0;
+
+	/**
+	 * Called by the engine when external collisions are no longer necessary or when they need to be updated
+	 * with some of the previous collisions removed. It is recommended in derived simulations to avoid freeing
+	 * any allocations regarding external collisions as it is likely more will be added soon after this call
+	 */
+	virtual void ClearExternalCollisions() = 0;
+
+	/**
+	 * Called by the engine to request data on all active collisions in a simulation. if bIncludeExternal is
+	 * true, derived implementations should add the asset collisions and any collisions added at runtime, if
+	 * false only the collisions from the asset should be considered
+	 * @param OutCollisions - Array to write collisions to
+	 * @param bIncludeExternal - Whether or not external collisions should be retrieved, or just asset collisions
+	 */
+	virtual void GetCollisions(FClothCollisionData& OutCollisions, bool bIncludeExternal = true) const = 0;
+
+private:
+
+public:
+
+	/**
+	 * Called by the engine once per physics tick to gather non-cycle stats (num sim verts etc.)
+	 * Will not be called if STATS is not defined, not required to be implemented if no need for stats.
+	 * STAT_NumCloths and STAT_NumClothVerts are already defined and can be incremented, more can be
+	 * added per simulation.
+	 * This can be called on any thread, so derived implementations should make sure the gather is
+	 * safe
+	 */
+	virtual void GatherStats() const
+	{}
+};
diff --git a/Engine/Source/Runtime/ClothingSystemRuntimeNv/Private/ClothingSimulationNv.cpp b/Engine/Source/Runtime/ClothingSystemRuntimeNv/Private/ClothingSimulationNv.cpp
--- a/Engine/Source/Runtime/ClothingSystemRuntimeNv/Private/ClothingSimulationNv.cpp
+++ b/Engine/Source/Runtime/ClothingSystemRuntimeNv/Private/ClothingSimulationNv.cpp
@@ -26,7 +26,6 @@
 #endif
 
 #include "PhysicsEngine/PhysicsAsset.h"
-#include "PhysicsEngine/PhysicsSettings.h"
 #include "SkeletalRenderPublic.h"
 
 DECLARE_CYCLE_STAT(TEXT("Compute Clothing Normals"), STAT_NvClothComputeNormals, STATGROUP_Physics);
@@ -35,142 +34,63 @@ DECLARE_CYCLE_STAT(TEXT("Update Collisions"), STAT_NvClothUpdateCollisions, STAT
 DECLARE_CYCLE_STAT(TEXT("Fill Context"), STAT_NvClothFillContext, STATGROUP_Physics);
 DECLARE_CYCLE_STAT(TEXT("Update Anim Drive"), STAT_NvClothUpdateAnimDrive, STATGROUP_Physics);
 
-//=============================================================================
-// FClothingSimulationBase
-//=============================================================================
-
 static TAutoConsoleVariable<float> GClothMaxDeltaTimeTeleportMultiplier(
 	TEXT("p.Cloth.MaxDeltaTimeTeleportMultiplier"),
 	1.5f,
 	TEXT("A multiplier of the MaxPhysicsDelta time at which we will automatically just teleport cloth to its new location\n")
 	TEXT(" default: 1.5"));
 
-FClothingSimulationBase::FClothingSimulationBase()
+//=============================================================================
+// FClothingSimulationContextNv
+//=============================================================================
+
+FClothingSimulationContextNv::FClothingSimulationContextNv()
+	: PredictedLod(0)
+	, WindAdaption(0.0f)
+	, TeleportMode(EClothingTeleportMode::None)
+	, MaxDistanceScale(1.f)
 {
-	MaxPhysicsDelta = UPhysicsSettings::Get()->MaxPhysicsDeltaTime;
 }
 
-FClothingSimulationBase::~FClothingSimulationBase()
-{}
+FClothingSimulationContextNv::~FClothingSimulationContextNv()
+{
+}
 
-void FClothingSimulationBase::FillContext(
-	USkeletalMeshComponent* InComponent, 
-	float InDeltaTime, 
-	IClothingSimulationContext* InOutContext)
+void FClothingSimulationContextNv::Fill(const USkeletalMeshComponent* InComponent, float InDeltaSeconds, float InMaxPhysicsDelta)
 {
+	SCOPE_CYCLE_COUNTER(STAT_NvClothFillContext);
 	LLM_SCOPE(ELLMTag::SkeletalMesh);
 
-	FClothingSimulationContextBase* BaseContext = static_cast<FClothingSimulationContextBase*>(InOutContext);
-	BaseContext->ComponentToWorld = InComponent->GetComponentTransform();
-	BaseContext->PredictedLod = InComponent->PredictedLODLevel;
-	InComponent->GetWindForCloth_GameThread(BaseContext->WindVelocity, BaseContext->WindAdaption);
-	USkeletalMesh* SkelMesh = InComponent->SkeletalMesh;
+	FClothingSimulationContextCommon::Fill(InComponent, InDeltaSeconds, InMaxPhysicsDelta);
 
-	if(USkinnedMeshComponent* MasterComponent = InComponent->MasterPoseComponent.Get())
-	{
-		int32 NumBones = InComponent->MasterBoneMap.Num();
+	PredictedLod = InComponent->PredictedLODLevel;
 
-		if(NumBones == 0)
-		{
-			if(InComponent->SkeletalMesh)
-			{
-				// This case indicates an invalid master pose component (e.g. no skeletal mesh)
-				NumBones = InComponent->SkeletalMesh->RefSkeleton.GetNum();
+	TeleportMode = (InDeltaSeconds > InMaxPhysicsDelta * GClothMaxDeltaTimeTeleportMultiplier.GetValueOnGameThread()) ?
+		EClothingTeleportMode::Teleport :
+		InComponent->ClothTeleportMode;
 
-				BaseContext->BoneTransforms.Empty(NumBones);
-				BaseContext->BoneTransforms.AddDefaulted(NumBones);
-			}
-		}
-		else
-		{
-		    BaseContext->BoneTransforms.Reset(NumBones);
-		    BaseContext->BoneTransforms.AddDefaulted(NumBones);
-    
-		    for(int32 BoneIndex = 0; BoneIndex < NumBones; ++BoneIndex)
-		    {
-			    bool bFoundMaster = false;
-			    if(InComponent->MasterBoneMap.IsValidIndex(BoneIndex))
-			    {
-				    const int32 MasterIndex = InComponent->MasterBoneMap[BoneIndex];
-    
-				    if(MasterIndex != INDEX_NONE && MasterIndex < MasterComponent->GetComponentSpaceTransforms().Num())
-				    {
-					    BaseContext->BoneTransforms[BoneIndex] = MasterComponent->GetComponentSpaceTransforms()[MasterIndex];
-					    bFoundMaster = true;
-				    }
-			    }
-    
-			    if(!bFoundMaster)
-			    {
-				    if(SkelMesh)
-				    {
-					    const int32 ParentIndex = SkelMesh->RefSkeleton.GetParentIndex(BoneIndex);
-    
-					    if(ParentIndex != INDEX_NONE)
-					    {
-						    BaseContext->BoneTransforms[BoneIndex] = BaseContext->BoneTransforms[ParentIndex] * SkelMesh->RefSkeleton.GetRefBonePose()[BoneIndex];
-					    }
-					    else
-					    {
-						    BaseContext->BoneTransforms[BoneIndex] = SkelMesh->RefSkeleton.GetRefBonePose()[BoneIndex];
-					    }
-				    }
-			    }
-		    }
-		}
-	}
-	else
-	{
-		BaseContext->BoneTransforms = InComponent->GetComponentSpaceTransforms();
-	}
-
-	UWorld* ComponentWorld = InComponent->GetWorld();
-	check(ComponentWorld);
-
-	BaseContext->DeltaSeconds = FMath::Min(InDeltaTime, MaxPhysicsDelta);
-
-	if(InDeltaTime > (MaxPhysicsDelta * GClothMaxDeltaTimeTeleportMultiplier.GetValueOnGameThread()))
-	{
-		BaseContext->TeleportMode = EClothingTeleportMode::Teleport;
-	}
-	else
-	{
-		BaseContext->TeleportMode = InComponent->ClothTeleportMode;
-	}	
-
-	BaseContext->MaxDistanceScale = InComponent->GetClothMaxDistanceScale();
-
-	float GravityStrength = 981.0f;
-
-	if(UWorld* CurrWorld = InComponent->GetWorld())
-	{
-		GravityStrength = CurrWorld->GetGravityZ();
-	}
+	MaxDistanceScale = const_cast<USkeletalMeshComponent*>(InComponent)->GetClothMaxDistanceScale();  // TODO(Kriss.Gossart): Test in BP whether GetClothMaxDistanceScale can be made const without breaking any functionalities
+}
 
-	BaseContext->WorldGravity = FVector(0.0f, 0.0f, GravityStrength);
+void FClothingSimulationContextNv::FillRefToLocals(const USkeletalMeshComponent* InComponent)
+{
+	FClothingSimulationContextCommon::RefToLocals.Reset();
+	InComponent->GetCurrentRefToLocalMatrices(FClothingSimulationContextCommon::RefToLocals, InComponent->PredictedLODLevel);
+}
 
-	// Checking the component here to track rare issue leading to invalid contexts
-	if(InComponent->IsPendingKill())
+void FClothingSimulationContextNv::FillWorldGravity(const USkeletalMeshComponent* InComponent)
+{
+	float GravityStrength = 981.f;  // This seems to be pointing the wrong way around, but since this is legacy code it is probably safer not to change it
+	if (const UWorld* const ComponentWorld = InComponent->GetWorld())
 	{
-		AActor* CompOwner = InComponent->GetOwner();
-		UE_LOG(LogSkeletalMesh, Warning, 
-			TEXT("Attempting to fill a clothing simulation context for a PendingKill skeletal mesh component (Comp: %s, Actor: %s). "
-				"Pending kill skeletal mesh components should be unregistered before marked pending kill."), 
-			*InComponent->GetName(), CompOwner ? *CompOwner->GetName() : TEXT("None"));
-
-		// Make sure we clear this out to skip any attempted simulations
-		BaseContext->BoneTransforms.Reset();
+		GravityStrength = ComponentWorld->GetGravityZ();
 	}
+	FClothingSimulationContextCommon::WorldGravity = FVector(0.f, 0.f, GravityStrength);
+}
 
-	if(BaseContext->BoneTransforms.Num() == 0)
-	{
-		AActor* CompOwner = InComponent->GetOwner();
-		USkinnedMeshComponent* Master = InComponent->MasterPoseComponent.Get();
-		UE_LOG(LogSkeletalMesh, Warning, TEXT("Attempting to fill a clothing simulation context for a skeletal mesh component that has zero bones (Comp: %s, Master: %s, Actor: %s)."), *InComponent->GetName(), Master ? *Master->GetName() : TEXT("None"), CompOwner ? *CompOwner->GetName() : TEXT("None"));
-
-		// Make sure we clear this out to skip any attempted simulations
-		BaseContext->BoneTransforms.Reset();
-	}
+void FClothingSimulationContextNv::FillWindVelocity(const USkeletalMeshComponent* InComponent)
+{
+	WindAdaption = FClothingSimulationContextCommon::SetWindFromComponent(InComponent);  // Also update wind adaptation
 }
 
 //=============================================================================
@@ -599,22 +519,6 @@ IClothingSimulationContext* FClothingSimulationNv::CreateContext()
 	return new FClothingSimulationContextNv();
 }
 
-void FClothingSimulationNv::FillContext(USkeletalMeshComponent* InComponent, float InDeltaTime, IClothingSimulationContext* InOutContext)
-{
-	SCOPE_CYCLE_COUNTER(STAT_NvClothFillContext);
-
-	check(InOutContext);
-
-	FClothingSimulationBase::FillContext(InComponent, InDeltaTime, InOutContext);
-
-	// Assume calling code guarantees the safety of this conversion, we should be given the pointer we allocated in CreateContext
-	FClothingSimulationContextNv* NvContext = (FClothingSimulationContextNv*)InOutContext;
-
-	// Get the current ref to locals to skin fixed vertices
-	NvContext->RefToLocals.Reset();
-	InComponent->GetCurrentRefToLocalMatrices(NvContext->RefToLocals, InComponent->PredictedLODLevel);
-}
-
 void FClothingSimulationNv::Initialize()
 {
 	FClothingSystemRuntimeModuleNv& ClothingModule = FModuleManager::Get().LoadModuleChecked<FClothingSystemRuntimeModuleNv>("ClothingSystemRuntimeNv");
diff --git a/Engine/Source/Runtime/ClothingSystemRuntimeNv/Public/ClothingSimulationNv.h b/Engine/Source/Runtime/ClothingSystemRuntimeNv/Public/ClothingSimulationNv.h
--- a/Engine/Source/Runtime/ClothingSystemRuntimeNv/Public/ClothingSimulationNv.h
+++ b/Engine/Source/Runtime/ClothingSystemRuntimeNv/Public/ClothingSimulationNv.h
@@ -5,9 +5,9 @@
 #if WITH_NVCLOTH
 
 #include "ClothingSimulation.h"
-#include "ClothingSimulationContext.h"
 #include "ClothingActor.h"
 #include "ClothConfigNv.h"
+#include "ClothCollisionData.h"
 
 class UClothingAssetBase;
 class UClothingAssetCommon;
@@ -29,13 +29,36 @@ namespace physx
 	class PxVec4;
 }
 
-class FClothingSimulationContextNv final : public FClothingSimulationContextBase
+class FClothingSimulationContextNv final : public FClothingSimulationContextCommon
 {
 public:
+	FClothingSimulationContextNv();
+	virtual ~FClothingSimulationContextNv() override;
 
-	// Ref to local matrices from the owning component (for skinning fixed verts)
-	TArray<FMatrix> RefToLocals;
+	// Override the fill context function to also set the Nv specific simulation context members
+	virtual void Fill(const USkeletalMeshComponent* InComponent, float InDeltaSeconds, float InMaxPhysicsDelta) override;
 
+	// Set the RefToLocals array in the parent class using Nv specific predicted LOD information
+	virtual void FillRefToLocals(const USkeletalMeshComponent* InComponent) override;
+
+	// Set the world gravity in the parent class while preserving the Nv legacy code behavior
+	void virtual FillWorldGravity(const USkeletalMeshComponent* InComponent) override;
+
+	// Set WindVelocity in the parent class and Nv specific WindAdaption
+	virtual void FillWindVelocity(const USkeletalMeshComponent* InComponent) override;
+
+	// The predicted LOD of the skeletal mesh component running the simulation
+	int32 PredictedLod;
+
+	// Wind adaption, a measure of how quickly to adapt to the wind speed
+	// when using the legacy wind calculation mode
+	float WindAdaption;
+
+	// Whether and how we should teleport the simulation this tick
+	EClothingTeleportMode TeleportMode;
+
+	// Scale for the max distance constraints of the simulation mesh
+	float MaxDistanceScale;
 };
 
 // Scratch data for simulation to avoid allocations while processing, per actor data
@@ -162,26 +185,7 @@ private:
 	friend class FClothingSimulationNv;
 };
 
-// Base simulation to fill in common data for the base context
-class FClothingSimulationBase : public IClothingSimulation
-{
-public:
-	FClothingSimulationBase();
-	virtual ~FClothingSimulationBase();
-
-protected:
-	/** Fills in the base data for a clothing simulation */
-	virtual void 
-	FillContext(
-		USkeletalMeshComponent* InComponent, 
-		float InDeltaTime, 
-		IClothingSimulationContext* InOutContext) override;
-
-	/** Maximum physics time, incoming deltas will be clamped down to this value on long frames */
-	float MaxPhysicsDelta;
-};
-
-class FClothingSimulationNv final : public FClothingSimulationBase
+class FClothingSimulationNv final : public FClothingSimulationCommon
 {
 	// Cached from the module for speed, DO NOT DELETE, only for creating cloth objects
 	nv::cloth::Factory* CachedFactory;
@@ -201,7 +205,6 @@ public:
 	virtual void CreateActor(USkeletalMeshComponent* InOwnerComponent, UClothingAssetBase* InAsset, int32 InSimDataIndex) override;
 	virtual void PostActorCreationInitialize() override {};
 	virtual IClothingSimulationContext* CreateContext() override;
-	virtual void FillContext(USkeletalMeshComponent* InComponent, float InDeltaTime, IClothingSimulationContext* InOutContext) override;
 	virtual void Initialize() override;
 	virtual void Shutdown() override;
 	virtual bool ShouldSimulate() const override;
diff --git a/Engine/Source/Runtime/Engine/Classes/Components/SkeletalMeshComponent.h b/Engine/Source/Runtime/Engine/Classes/Components/SkeletalMeshComponent.h
--- a/Engine/Source/Runtime/Engine/Classes/Components/SkeletalMeshComponent.h
+++ b/Engine/Source/Runtime/Engine/Classes/Components/SkeletalMeshComponent.h
@@ -23,6 +23,7 @@
 #include "ClothingSystemRuntimeTypes.h"
 #include "ClothingSimulationInterface.h"
 #include "ClothingSimulationFactory.h"
+#include "ClothCollisionPrim.h"
 #include "PhysicsEngine/PhysicsAsset.h"
 
 #include "SkeletalMeshComponent.generated.h"
@@ -1693,7 +1694,7 @@ public:
 
 	FOnBoneTransformsFinalized OnBoneTransformsFinalized;
 
-	void GetCurrentRefToLocalMatrices(TArray<FMatrix>& OutRefToLocals, int32 InLodIdx);
+	void GetCurrentRefToLocalMatrices(TArray<FMatrix>& OutRefToLocals, int32 InLodIdx) const;
 
 	// Conditions used to gate when post process events happen
 	bool ShouldUpdatePostProcessInstance() const;
@@ -2143,7 +2144,7 @@ private:
 	/** See UpdateClothTransform for documentation. */
 	void UpdateClothTransformImp();
 
-	friend class FClothingSimulationBase;
+	friend class FClothingSimulationContextCommon;
 
 	friend class FTickClothingTask;
 
diff --git a/Engine/Source/Runtime/Engine/Private/Components/SkeletalMeshComponent.cpp b/Engine/Source/Runtime/Engine/Private/Components/SkeletalMeshComponent.cpp
--- a/Engine/Source/Runtime/Engine/Private/Components/SkeletalMeshComponent.cpp
+++ b/Engine/Source/Runtime/Engine/Private/Components/SkeletalMeshComponent.cpp
@@ -3733,7 +3733,7 @@ void USkeletalMeshComponent::FinalizeBoneTransform()
 	TRACE_SKELETAL_MESH_COMPONENT(this);
 }
 
-void USkeletalMeshComponent::GetCurrentRefToLocalMatrices(TArray<FMatrix>& OutRefToLocals, int32 InLodIdx)
+void USkeletalMeshComponent::GetCurrentRefToLocalMatrices(TArray<FMatrix>& OutRefToLocals, int32 InLodIdx) const
 {
 	if(SkeletalMesh)
 	{
diff --git a/Engine/Source/Runtime/Engine/Private/SkeletalMeshComponentPhysics.cpp b/Engine/Source/Runtime/Engine/Private/SkeletalMeshComponentPhysics.cpp
--- a/Engine/Source/Runtime/Engine/Private/SkeletalMeshComponentPhysics.cpp
+++ b/Engine/Source/Runtime/Engine/Private/SkeletalMeshComponentPhysics.cpp
@@ -25,6 +25,7 @@
 #include "Rendering/SkeletalMeshRenderData.h"
 #include "Physics/PhysicsInterfaceCore.h"
 #include "AnimationRuntime.h"
+#include "ClothCollisionData.h"
 
 #include "Logging/MessageLog.h"
 #include "CollisionDebugDrawingPublic.h"
