Chaos Cloth - Make USkeletalMeshComponent::GetClothMaxDistanceScale() const, after having tested that changing a UFunction to const doesn't break the blueprint assets using it.

diff --git a/Engine/Source/Runtime/ClothingSystemRuntimeNv/Private/ClothingSimulationNv.cpp b/Engine/Source/Runtime/ClothingSystemRuntimeNv/Private/ClothingSimulationNv.cpp
--- a/Engine/Source/Runtime/ClothingSystemRuntimeNv/Private/ClothingSimulationNv.cpp
+++ b/Engine/Source/Runtime/ClothingSystemRuntimeNv/Private/ClothingSimulationNv.cpp
@@ -69,7 +69,7 @@ void FClothingSimulationContextNv::Fill(const USkeletalMeshComponent* InComponen
 		EClothingTeleportMode::Teleport :
 		InComponent->ClothTeleportMode;
 
-	MaxDistanceScale = const_cast<USkeletalMeshComponent*>(InComponent)->GetClothMaxDistanceScale();  // TODO(Kriss.Gossart): Test in BP whether GetClothMaxDistanceScale can be made const without breaking any functionalities
+	MaxDistanceScale = InComponent->GetClothMaxDistanceScale();
 }
 
 void FClothingSimulationContextNv::FillRefToLocals(const USkeletalMeshComponent* InComponent)
diff --git a/Engine/Source/Runtime/Engine/Classes/Components/SkeletalMeshComponent.h b/Engine/Source/Runtime/Engine/Classes/Components/SkeletalMeshComponent.h
--- a/Engine/Source/Runtime/Engine/Classes/Components/SkeletalMeshComponent.h
+++ b/Engine/Source/Runtime/Engine/Classes/Components/SkeletalMeshComponent.h
@@ -1010,7 +1010,7 @@ public:
 	 * Get/Set the max distance scale of clothing mesh vertices
 	 */
 	UFUNCTION(BlueprintCallable, Category="Components|SkeletalMesh")
-	float GetClothMaxDistanceScale();
+	float GetClothMaxDistanceScale() const;
 	UFUNCTION(BlueprintCallable, Category="Components|SkeletalMesh")
 	void SetClothMaxDistanceScale(float Scale);
 
diff --git a/Engine/Source/Runtime/Engine/Private/SkeletalMeshComponentPhysics.cpp b/Engine/Source/Runtime/Engine/Private/SkeletalMeshComponentPhysics.cpp
--- a/Engine/Source/Runtime/Engine/Private/SkeletalMeshComponentPhysics.cpp
+++ b/Engine/Source/Runtime/Engine/Private/SkeletalMeshComponentPhysics.cpp
@@ -3337,7 +3337,7 @@ void USkeletalMeshComponent::SetAllUseCCD(bool InUseCCD)
 }
 
 // blueprint callable methods 
-float USkeletalMeshComponent::GetClothMaxDistanceScale()
+float USkeletalMeshComponent::GetClothMaxDistanceScale() const
 {
 	return ClothMaxDistanceScale;
 }
