#UE4 Forward declare fix

diff --git a/Engine/Source/Runtime/ClothingSystemRuntimeNv/Public/ClothingSimulationNv.h b/Engine/Source/Runtime/ClothingSystemRuntimeNv/Public/ClothingSimulationNv.h
--- a/Engine/Source/Runtime/ClothingSystemRuntimeNv/Public/ClothingSimulationNv.h
+++ b/Engine/Source/Runtime/ClothingSystemRuntimeNv/Public/ClothingSimulationNv.h
@@ -9,6 +9,7 @@
 #include "ClothConfigNv.h"
 #include "ClothCollisionData.h"
 
+class FPrimitiveDrawInterface;
 class UClothingAssetBase;
 class UClothingAssetCommon;
 
