Fix missing include
#rnx

diff --git a/Engine/Source/Runtime/ClothingSystemRuntimeNv/Private/ClothingSimulationFactoryNv.cpp b/Engine/Source/Runtime/ClothingSystemRuntimeNv/Private/ClothingSimulationFactoryNv.cpp
--- a/Engine/Source/Runtime/ClothingSystemRuntimeNv/Private/ClothingSimulationFactoryNv.cpp
+++ b/Engine/Source/Runtime/ClothingSystemRuntimeNv/Private/ClothingSimulationFactoryNv.cpp
@@ -3,6 +3,7 @@
 #include "ClothingSimulationFactoryNv.h"
 
 #if WITH_NVCLOTH
+#include "UObject/Package.h"
 #include "ClothPhysicalMeshData.h"  // For EWeightMapTargetCommon
 #include "ClothingSimulationNv.h"
 #include "ClothingSimulationInteractorNv.h"
