CIS - Fix UClothingSimulationInteractorNv cannot instantiate abstract class error in NOPCH builds.

diff --git a/Engine/Source/Runtime/ClothingSystemRuntimeNv/Public/ClothingSimulationInteractorNv.h b/Engine/Source/Runtime/ClothingSystemRuntimeNv/Public/ClothingSimulationInteractorNv.h
--- a/Engine/Source/Runtime/ClothingSystemRuntimeNv/Public/ClothingSimulationInteractorNv.h
+++ b/Engine/Source/Runtime/ClothingSystemRuntimeNv/Public/ClothingSimulationInteractorNv.h
@@ -30,6 +30,7 @@ public:
 
 	// TODO: These new functions are currently unimplemented
 	virtual void SetNumIterations(int32 /*NumIterations*/) override {}
+	virtual void SetMaxNumIterations(int32 /*MaxNumIterations*/) override {}
 	virtual void SetNumSubsteps(int32 /*NumSubsteps*/) override {}
 	//////////////////////////////////////////////////////////////////////////
 
