Mac CIS fix to physics interface cleanup

diff --git a/Engine/Source/Runtime/ClothingSystemRuntimeNv/Public/ClothingSimulationFactoryNv.h b/Engine/Source/Runtime/ClothingSystemRuntimeNv/Public/ClothingSimulationFactoryNv.h
--- a/Engine/Source/Runtime/ClothingSystemRuntimeNv/Public/ClothingSimulationFactoryNv.h
+++ b/Engine/Source/Runtime/ClothingSystemRuntimeNv/Public/ClothingSimulationFactoryNv.h
@@ -7,8 +7,10 @@
 
 #include "ClothingSimulationFactoryNv.generated.h"
 
+class UE_DEPRECATED(5.1, "NvCloth simulation is no longer supported, UChaosClothingSimulationFactory should be used going forward.") UClothingSimulationFactoryNv;
+
 UCLASS()
-class UE_DEPRECATED(5.1, "NvCloth simulation is no longer supported, UChaosClothingSimulationFactory should be used going forward.") CLOTHINGSYSTEMRUNTIMENV_API UClothingSimulationFactoryNv final : public UClothingSimulationFactory
+class CLOTHINGSYSTEMRUNTIMENV_API UClothingSimulationFactoryNv final : public UClothingSimulationFactory
 {
 	GENERATED_BODY()
 public:
