Cloth - Removed legacy defunct (nullptr) nvcloth simulation factory from the list of registered simulation factories.

#rnx

diff --git a/Engine/Source/Runtime/ClothingSystemRuntimeNv/Private/ClothingSystemRuntimeModuleNv.cpp b/Engine/Source/Runtime/ClothingSystemRuntimeNv/Private/ClothingSystemRuntimeModuleNv.cpp
--- a/Engine/Source/Runtime/ClothingSystemRuntimeNv/Private/ClothingSystemRuntimeModuleNv.cpp
+++ b/Engine/Source/Runtime/ClothingSystemRuntimeNv/Private/ClothingSystemRuntimeModuleNv.cpp
@@ -19,12 +19,10 @@ FClothingSystemRuntimeModuleNv::FClothingSystemRuntimeModuleNv()
 
 void FClothingSystemRuntimeModuleNv::StartupModule()
 {
-	IModularFeatures::Get().RegisterModularFeature(IClothingSimulationFactoryClassProvider::FeatureName, this);
 }
 
 void FClothingSystemRuntimeModuleNv::ShutdownModule()
 {
-	IModularFeatures::Get().UnregisterModularFeature(IClothingSimulationFactoryClassProvider::FeatureName, this);
 }
 
 nv::cloth::Factory* FClothingSystemRuntimeModuleNv::GetSoftwareFactory()
